Anim Node system initial check in to start iterating on.
- Anim Node Modules
- AnimNode Base Data/Instance, Interfaces, RigVM AnimNode update function
- Transitions + Basic Blend Transition
- Sequence Player

#Synced-CL 50686245

diff --git a/Engine/Plugins/Experimental/UAF/UAFAnimNode/Source/UAFAnimNode/Internal/AnimNode/RigUnit_RunAnimNode.h b/Engine/Plugins/Experimental/UAF/UAFAnimNode/Source/UAFAnimNode/Internal/AnimNode/RigUnit_RunAnimNode.h
new file mode 100644
--- /dev/null
+++ b/Engine/Plugins/Experimental/UAF/UAFAnimNode/Source/UAFAnimNode/Internal/AnimNode/RigUnit_RunAnimNode.h
@@ -0,0 +1,69 @@
+// Copyright Epic Games, Inc. All Rights Reserved.
+
+#pragma once
+
+#include "AnimationRuntime.h"
+#include "AnimNextExecuteContext.h"
+#include "Animation/AnimSequence.h"
+#include "EvaluationVM/EvaluationProgram.h"
+#include "Factory/AnimAssetParams.h"
+#include "Graph/AnimNext_LODPose.h"
+#include "Graph/RigUnit_AnimNextBase.h"
+#include "AnimNode/UAFAnimNode.h"
+
+#include "RigUnit_RunAnimNode.generated.h"
+
+#define UE_API UAFANIMNODE_API
+
+	
+USTRUCT()
+struct FUAFAnimNodeReference
+{
+	GENERATED_BODY()
+	UPROPERTY(EditAnywhere, Category = "Run", meta = (Input, ExcludeBaseStruct))
+	TInstancedStruct<UE::UAF::FUAFAnimNodeData> NodeData;
+};
+
+USTRUCT()
+struct FRunAnimNodeWorkData
+{
+	GENERATED_BODY()
+
+	TSharedPtr<UE::UAF::FUAFAnimNodeInstance> AnimNodeInstance;
+};
+
+/** Runs an AnimNode */
+USTRUCT(meta=(DisplayName="Run Anim Node", Category="Anim Node", NodeColor="0, 1, 1", Keywords="Animation, Node, Graph, AnimNode"))
+struct FRigUnit_RunAnimNode : public FRigUnit_AnimNextBase
+{
+	GENERATED_BODY()
+
+	RIGVM_METHOD()
+	UE_API void Execute();
+
+	// The AnimNode to run
+	UPROPERTY(EditAnywhere, Category = "Run", meta = (Input))
+	FUAFAnimNodeReference AnimNode;
+
+	// LOD to run at. If this is -1 then the reference pose's source LOD will be used
+	UPROPERTY(EditAnywhere, Category = "Run", meta = (Input))
+	int32 LOD = -1;
+
+	// Reference pose
+	UPROPERTY(EditAnywhere, Category = "Run", meta = (Input))
+	FAnimNextGraphReferencePose ReferencePose;
+
+	// Pose result
+	UPROPERTY(EditAnywhere, Category = "Run", meta = (Output))
+	FAnimNextGraphLODPose Result;
+
+	// Internal work data
+	UPROPERTY(transient)
+	FRunAnimNodeWorkData WorkData;
+
+	// The execution result
+	UPROPERTY(EditAnywhere, DisplayName = "Execute", Category = "BeginExecution", meta = (Input, Output))
+	FAnimNextExecuteContext ExecuteContext;
+};
+
+#undef UE_API
diff --git a/Engine/Plugins/Experimental/UAF/UAFAnimNode/Source/UAFAnimNode/Internal/AnimNode/UAFSequencePlayerNode.h b/Engine/Plugins/Experimental/UAF/UAFAnimNode/Source/UAFAnimNode/Internal/AnimNode/UAFSequencePlayerNode.h
new file mode 100644
--- /dev/null
+++ b/Engine/Plugins/Experimental/UAF/UAFAnimNode/Source/UAFAnimNode/Internal/AnimNode/UAFSequencePlayerNode.h
@@ -0,0 +1,80 @@
+// Copyright Epic Games, Inc. All Rights Reserved.
+
+#pragma once
+
+#include "AnimationRuntime.h"
+#include "Animation/AnimSequence.h"
+#include "EvaluationVM/EvaluationProgram.h"
+#include "AnimNode/UAFAnimNode.h"
+#include "Factory/AnimAssetParams.h" // for EAnimAssetLoopMode
+
+#include "UAFSequencePlayerNode.generated.h"
+
+namespace UE::UAF
+{
+	
+// AnimNode Instance for playing an AnimSequence 
+struct FUAFSequencePlayerNodeData;
+struct FUAFSequencePlayerNodeInstance : public FUAFAnimNodeInstance, public IUAFAnimNodeTimeline
+{
+	FUAFSequencePlayerNodeInstance(const FUAFSequencePlayerNodeData* InData);
+
+	virtual void Update(FUAFAnimNodeUpdateContext& Context, float DeltaTime) override;
+
+	virtual void* GetInterface(FAnimNodeInterfaceId Id) override;
+
+	// IUAFAnimNodeTimeline
+	virtual float GetCurrentTime() override;
+	virtual void SetCurrentTime(float InCurrentTime) override;
+	virtual bool IsLooping() override;
+	virtual float GetLength() override;
+
+private:
+	/** Delta time range required for root motion extraction **/
+	FDeltaTimeRecord DeltaTimeRecord;
+	float CurrentTime = 0.0f;
+
+	TObjectPtr<const UAnimSequence> Sequence;
+	bool bIsLooping;
+};
+
+
+// AnimNode Data for playing an AnimSequence
+USTRUCT()
+struct FUAFSequencePlayerNodeData : public FUAFAnimNodeData
+#if CPP
+, public IUAFAnimNodeDataHasAsset
+#endif
+{
+	GENERATED_BODY()
+
+	UPROPERTY(EditAnywhere, Category="Data")
+	TObjectPtr<UAnimSequence> Sequence;
+
+	// HasAsset interface
+	virtual UObject* GetAsset() const override { return Sequence; } 
+
+	/** Whether to use the asset's default looping behavior or force an override. */
+	UPROPERTY(EditAnywhere, Category = "Data")
+	EAnimAssetLoopMode LoopMode = EAnimAssetLoopMode::Auto;
+
+	bool IsLooping() const
+	{
+		switch (LoopMode)
+		{
+		case EAnimAssetLoopMode::Auto:
+			return Sequence ? Sequence->bLoop : false;
+		case EAnimAssetLoopMode::ForceLoop:
+			return true;
+		case EAnimAssetLoopMode::ForceNonLoop:
+			return false;
+		default:
+			return false;
+		}
+	}
+
+	virtual TSharedPtr<FUAFAnimNodeInstance> CreateInstance(FUAFAnimNodeUpdateContext& Context) const override;
+	virtual void* GetInterface(FAnimNodeInterfaceId Id) override;
+};
+	
+}
\ No newline at end of file
diff --git a/Engine/Plugins/Experimental/UAF/UAFAnimNode/Source/UAFAnimNode/Private/Module.cpp b/Engine/Plugins/Experimental/UAF/UAFAnimNode/Source/UAFAnimNode/Private/Module.cpp
new file mode 100644
--- /dev/null
+++ b/Engine/Plugins/Experimental/UAF/UAFAnimNode/Source/UAFAnimNode/Private/Module.cpp
@@ -0,0 +1,27 @@
+// Copyright Epic Games, Inc. All Rights Reserved.
+
+#include "CoreMinimal.h"
+#include "Module/AnimNextModule.h"
+#include "Modules/ModuleInterface.h"
+#include "Modules/ModuleManager.h"
+#include "RigVMCore/RigVMRegistry.h"
+
+namespace UE::UAF::AnimNode
+{
+
+class FModule : public IModuleInterface
+{
+public:
+	virtual void StartupModule() override
+	{
+	}
+
+	virtual void ShutdownModule() override
+	{
+	}
+
+};
+
+}
+
+IMPLEMENT_MODULE(UE::UAF::AnimNode::FModule, UAFAnimNode)
diff --git a/Engine/Plugins/Experimental/UAF/UAFAnimNode/Source/UAFAnimNode/Private/RigUnit_RunAnimNode.cpp b/Engine/Plugins/Experimental/UAF/UAFAnimNode/Source/UAFAnimNode/Private/RigUnit_RunAnimNode.cpp
new file mode 100644
--- /dev/null
+++ b/Engine/Plugins/Experimental/UAF/UAFAnimNode/Source/UAFAnimNode/Private/RigUnit_RunAnimNode.cpp
@@ -0,0 +1,100 @@
+// Copyright Epic Games, Inc. All Rights Reserved.
+
+#include "AnimNode/RigUnit_RunAnimNode.h"
+
+#include "Module/AnimNextModuleInstance.h"
+#include "UAF/UAFSkeletonUserData.h"
+#include "UAFLogging.h"
+#include "EvaluationVM/EvaluationVM.h"
+#include "EvaluationVM/Tasks/PushAnimSequenceKeyframe.h"
+#include "Graph/UAFAnimGraphTrace.h"
+#include "TraitInterfaces/IEvaluate.h"
+
+#include UE_INLINE_GENERATED_CPP_BY_NAME(RigUnit_RunAnimNode)
+
+FRigUnit_RunAnimNode_Execute()
+{
+	SCOPED_NAMED_EVENT(UAF_Run_AnimNode, FColor::Orange);
+
+	using namespace UE::UAF;
+
+	const FAnimNextModuleContextData& ModuleContextData = ExecuteContext.GetContextData<FAnimNextModuleContextData>();
+	FAnimNextModuleInstance& ModuleInstance = ModuleContextData.GetModuleInstance();
+
+	if (!ReferencePose.ReferencePose.IsValid())
+	{
+		UAF_RIGUNIT_LOG(Warning, TEXT("Could not run Anim Node - Ref Pose is invalid]")); 
+		return;
+	}
+	
+	if (!AnimNode.NodeData.IsValid())
+	{
+		UAF_RIGUNIT_LOG(Warning, TEXT("Could not run Anim Node - No Anim Node Selected]")); 
+		return;
+	}
+
+	const UE::UAF::FReferencePose& RefPose = ReferencePose.ReferencePose.GetRef<UE::UAF::FReferencePose>();
+
+	int32 DesiredLOD = LOD;
+	if (DesiredLOD == -1)
+	{
+		DesiredLOD = RefPose.GetSourceLODLevel();
+	}
+
+	// TODO: Currently forcing additive flag to false here
+	if (Result.LODPose.ShouldPrepareForLOD(RefPose, DesiredLOD, false))
+	{
+		Result.LODPose.PrepareForLOD(RefPose, DesiredLOD, true, false);
+	}
+
+	ensure(Result.LODPose.LODLevel == DesiredLOD);
+
+	// Internally we use memstack allocation, so we need a mark here
+	FMemStack& MemStack = FMemStack::Get();
+	FMemMark MemMark(MemStack);
+
+	// We allocate a dummy buffer to trigger the allocation of a large chunk if this is the first mark
+	// This reduces churn internally by avoiding a chunk to be repeatedly allocated and freed as we push/pop marks
+	MemStack.Alloc(size_t(FPageAllocator::SmallPageSize) + 1, 16);
+
+	FEvaluationProgram EvaluationProgram;
+	FUAFAnimNodeUpdateContext Context(EvaluationProgram, ModuleContextData.GetObject(), &ModuleInstance);
+
+	if (WorkData.AnimNodeInstance == nullptr)
+	{
+		// allocate AnimNode instance
+		WorkData.AnimNodeInstance = AnimNode.NodeData.Get().CreateInstance(Context);
+	}
+
+	TRACE_UAF_MODULE(ModuleInstance);
+
+	if (WorkData.AnimNodeInstance)
+	{
+		{
+			SCOPED_NAMED_EVENT(AnimNode_Update, FColor::Orange);
+			WorkData.AnimNodeInstance->Update(Context, ExecuteContext.GetDeltaTime());
+		}
+
+		{
+			SCOPED_NAMED_EVENT(AnimNode_ExecuteEvaluationProgram, FColor::Orange);
+			
+			FEvaluationVM EvaluationVM(EEvaluationFlags::All, RefPose, nullptr, NAME_None, DesiredLOD);
+			
+			if (!EvaluationProgram.IsEmpty())
+			{
+				EvaluationProgram.Execute(EvaluationVM);
+
+				TUniquePtr<FKeyframeState> EvaluatedKeyframe;
+				if (EvaluationVM.PopValue(KEYFRAME_STACK_NAME, EvaluatedKeyframe))
+				{
+					Result.LODPose.CopyFrom(EvaluatedKeyframe->Pose);
+					Result.Curves.CopyFrom(EvaluatedKeyframe->Curves);
+					Result.Attributes.CopyFrom(EvaluatedKeyframe->Attributes);
+				}
+
+				TRACE_UAF_EVALUATIONPROGRAM_WITHOWNER(EvaluationProgram, ModuleInstance, ModuleInstance.GetObject());
+			}
+		}
+	}
+		
+}
diff --git a/Engine/Plugins/Experimental/UAF/UAFAnimNode/Source/UAFAnimNode/Private/UAFSequencePlayerNode.cpp b/Engine/Plugins/Experimental/UAF/UAFAnimNode/Source/UAFAnimNode/Private/UAFSequencePlayerNode.cpp
new file mode 100644
--- /dev/null
+++ b/Engine/Plugins/Experimental/UAF/UAFAnimNode/Source/UAFAnimNode/Private/UAFSequencePlayerNode.cpp
@@ -0,0 +1,94 @@
+// Copyright Epic Games, Inc. All Rights Reserved.
+
+#include "AnimNode/UAFSequencePlayerNode.h"
+
+#include "Module/AnimNextModuleInstance.h"
+#include "Graph/AnimNextAnimationGraph.h"
+#include "UAF/UAFSkeletonUserData.h"
+#include "UAFLogging.h"
+#include "EvaluationVM/EvaluationVM.h"
+#include "EvaluationVM/Tasks/PushAnimSequenceKeyframe.h"
+
+#include UE_INLINE_GENERATED_CPP_BY_NAME(UAFSequencePlayerNode)
+
+namespace UE::UAF
+{
+
+FUAFSequencePlayerNodeInstance::FUAFSequencePlayerNodeInstance(const FUAFSequencePlayerNodeData* InData)
+{
+	bIsLooping = InData->IsLooping();
+	Sequence = InData->Sequence;
+}
+
+void FUAFSequencePlayerNodeInstance::Update(FUAFAnimNodeUpdateContext& Context, float DeltaTime)
+{
+	SCOPED_NAMED_EVENT(AnimNode_Update_SequenceController, FColor::Green);
+	if (Sequence)
+	{
+		const float SequenceLength = Sequence->GetPlayLength();
+
+		DeltaTimeRecord.Set(CurrentTime, DeltaTime);
+		FAnimationRuntime::AdvanceTime(bIsLooping, DeltaTime, CurrentTime, SequenceLength);
+	
+		FAnimNextAnimSequenceKeyframeTask Task;
+		Task.AnimSequence = Sequence;
+		Task.SampleTime = CurrentTime;
+		Task.bInterpolate = true;
+		Task.DeltaTimeRecord = DeltaTimeRecord;
+		Task.bLooping = bIsLooping;
+		Task.bExtractTrajectory = true;
+
+		Context.AppendTask(Task);
+	}
+}
+
+void* FUAFSequencePlayerNodeInstance::GetInterface(FAnimNodeInterfaceId Id)
+{
+	if (Id == IUAFAnimNodeTimeline::InterfaceId)
+	{
+		return static_cast<IUAFAnimNodeTimeline*>(this);
+	}
+
+	return nullptr;
+}
+
+void FUAFSequencePlayerNodeInstance::SetCurrentTime(float InCurrentTime)
+{
+	CurrentTime = InCurrentTime;
+}
+
+bool FUAFSequencePlayerNodeInstance::IsLooping()
+{
+	return bIsLooping;
+}
+
+
+float FUAFSequencePlayerNodeInstance::GetCurrentTime()
+{
+	return CurrentTime;
+}
+
+float FUAFSequencePlayerNodeInstance::GetLength()
+{
+	if (Sequence)
+	{
+		return Sequence->GetPlayLength();
+	}
+	return 0;
+}
+
+TSharedPtr<FUAFAnimNodeInstance> FUAFSequencePlayerNodeData::CreateInstance(FUAFAnimNodeUpdateContext& Context) const
+{
+	return MakeShared<FUAFSequencePlayerNodeInstance>(this);
+}
+void* FUAFSequencePlayerNodeData::GetInterface(FAnimNodeInterfaceId Id)
+{
+	if (Id == IUAFAnimNodeDataHasAsset::InterfaceId)
+	{
+		return static_cast<IUAFAnimNodeDataHasAsset*>(this);
+	}
+
+	return nullptr;
+}
+	
+}
\ No newline at end of file
diff --git a/Engine/Plugins/Experimental/UAF/UAFAnimNode/Source/UAFAnimNode/Private/UAFTransition.cpp b/Engine/Plugins/Experimental/UAF/UAFAnimNode/Source/UAFAnimNode/Private/UAFTransition.cpp
new file mode 100644
--- /dev/null
+++ b/Engine/Plugins/Experimental/UAF/UAFAnimNode/Source/UAFAnimNode/Private/UAFTransition.cpp
@@ -0,0 +1,72 @@
+// Copyright Epic Games, Inc. All Rights Reserved.
+
+#include "AnimNode/UAFTransition.h"
+#include "EvaluationVM/Tasks/BlendKeyframes.h"
+#include "EvaluationVM/Tasks/PushReferenceKeyframe.h"
+
+#include UE_INLINE_GENERATED_CPP_BY_NAME(UAFTransition)
+
+namespace UE::UAF
+{
+	
+void FUAFBasicTransitionInstance::Update(UE::UAF::FUAFAnimNodeUpdateContext& Context, float DeltaTime)
+{
+	SCOPED_NAMED_EVENT(AnimNode_Update_UAFBasocTransitionNode, FColor::Blue);
+
+	TimeRemaining -= DeltaTime;
+	if (TimeRemaining<=0)
+	{
+		// if we had a special kind of shared ptr, in place of the normal TSharedPtr, we could actually prune ourselves automagically.
+		// AnimNodes would have to keep a list of Pointers to themselves (similar to the target of a weak pointer)
+		// and then when the transition completes we could set all the pointers to this transition to now point to the target controller instead.
+		Source = nullptr;
+		if (Target)
+		{
+			Target->Update(Context, DeltaTime);
+			return;
+		}
+	}
+	else
+	{
+		if (Source)
+		{
+			// manual pruning of nested transitions that are complete
+    		if(IUAFTransition* NestedTransition = Source->GetInterface<IUAFTransition>())
+    		{
+    			if (NestedTransition->IsComplete())
+    			{
+    				Source = NestedTransition->GetTarget();
+    			}
+    		}
+    		
+			if (Source && Target)
+			{
+				Source->Update(Context, DeltaTime);
+				Target->Update(Context, DeltaTime);
+
+				float BlendWeight = 1.0 - TimeRemaining/Duration;
+				Context.AppendTask(FAnimNextBlendTwoKeyframesTask::Make(BlendWeight));
+				return;
+			}
+		}
+	}
+	
+	Context.AppendTask(FAnimNextPushReferenceKeyframeTask::MakeFromSkeleton());
+}
+
+void* FUAFBasicTransitionInstance::GetInterface(FAnimNodeInterfaceId Id)
+{
+	if(Id == IUAFTransition::InterfaceId)
+	{
+		return static_cast<IUAFTransition*>(this);
+	}
+
+	if (Target)
+	{
+		return Target->GetInterface(Id);
+	}
+
+	return nullptr;
+}
+	
+}
\ No newline at end of file
diff --git a/Engine/Plugins/Experimental/UAF/UAFAnimNode/Source/UAFAnimNode/Public/AnimNode/UAFAnimNode.h b/Engine/Plugins/Experimental/UAF/UAFAnimNode/Source/UAFAnimNode/Public/AnimNode/UAFAnimNode.h
new file mode 100644
--- /dev/null
+++ b/Engine/Plugins/Experimental/UAF/UAFAnimNode/Source/UAFAnimNode/Public/AnimNode/UAFAnimNode.h
@@ -0,0 +1,143 @@
+// Copyright Epic Games, Inc. All Rights Reserved.
+
+#pragma once
+
+#include "AnimationRuntime.h"
+#include "AnimNextExecuteContext.h"
+#include "UAFAssetInstance.h"
+#include "EvaluationVM/EvaluationProgram.h"
+#include "ConstExprUID.h"
+
+#include "UAFAnimNode.generated.h"
+
+#define UE_API UAFANIMNODE_API
+
+namespace UE::UAF
+{
+	typedef FConstExprUID FAnimNodeInterfaceId;
+
+struct IUAFAnimNodeTimeline
+{
+	static constexpr FAnimNodeInterfaceId InterfaceId = FConstExprUID::MakeFromString(TEXT("IUAFAnimNodeTimeline"));
+	virtual float GetCurrentTime() = 0;
+	virtual float GetLength() = 0;
+	virtual bool IsLooping() = 0;
+	virtual void SetCurrentTime(float CurrentTime) = 0;
+};
+
+
+struct IUAFAnimNodeDataHasAsset
+{
+	static constexpr FAnimNodeInterfaceId InterfaceId = FConstExprUID::MakeFromString(TEXT("IUAFAnimNodeDataHasAsset"));
+	virtual UObject* GetAsset() const = 0;
+};
+	
+/**
+ * FUAFAnimNodeUpdateContext
+ *
+ * Contains an Evaluation program, which AnimNodes can add tasks to
+ */
+
+struct FUAFAnimNodeUpdateContext
+{
+	FUAFAnimNodeUpdateContext(FEvaluationProgram& InEvaluationProgram, UObject* InHostObject, FUAFAssetInstance* InVariablesOwner)
+	: HostObject(InHostObject), VariablesOwner(InVariablesOwner), EvaluationProgram(InEvaluationProgram) {};
+	
+	// Appends a new task into the evaluation program, tasks mutate state in the order they have been appended in
+	// This means that child nodes need to evaluate first, tasks will usually be appended in IEvaluate::PostEvaluate
+	// Tasks are moved into their final memory location, caller can allocate the task anywhere, it is no longer needed after this operation
+	// @see FEvaluationProgram, FEvaluationTask, FEvaluationVM
+	template<class TaskType>
+	void AppendTask(TaskType&& Task) { EvaluationProgram.AppendTask(MoveTemp(Task)); }
+
+	template<class TaskType>
+	void AppendTaskPtr(const TSharedPtr<TaskType>& TaskPtr) { EvaluationProgram.AppendTaskPtr(TaskPtr); }
+	
+	FUAFAssetInstance* GetVariablesOwner()
+	{
+		return VariablesOwner;
+	}
+
+	UObject* GetHostObject()
+	{
+		return HostObject;
+	}
+	
+private:
+	UObject* HostObject;
+	FUAFAssetInstance* VariablesOwner;
+
+	FEvaluationProgram& EvaluationProgram;
+};
+	
+
+/**
+ * FUAFAnimNodeInstance
+ *
+ * Animation Controller Instance
+ */
+struct FUAFAnimNodeInstance
+{	
+	virtual ~FUAFAnimNodeInstance()
+	{
+	}
+
+	virtual void Update(FUAFAnimNodeUpdateContext& Context, float DeltaTime)
+	{
+		ensureMsgf(false, TEXT("Animation Controller Instance Update Unimplemented!!"));
+	}
+
+	virtual void* GetInterface(FAnimNodeInterfaceId Id)
+	{
+		return nullptr;
+	}
+	
+	template <typename T>
+	T* GetInterface()
+	{
+		return static_cast<T*>(GetInterface(T::InterfaceId));
+	}
+};
+
+
+/**
+ * FUAFAnimNodeData
+ *
+ * Animation Controller Data
+ */
+USTRUCT()
+struct FUAFAnimNodeData
+{
+	GENERATED_BODY()
+	
+	virtual ~FUAFAnimNodeData()
+	{
+	}
+	
+	virtual TSharedPtr<FUAFAnimNodeInstance> CreateInstance(FUAFAnimNodeUpdateContext& Context) const
+	{
+		ensureMsgf(false, TEXT("AnimNode CreateInstance Unimplemented!!"));
+		return nullptr;
+	};
+
+	virtual void* GetInterface(FAnimNodeInterfaceId Id)
+	{
+		return nullptr;
+	}
+	
+	template <typename T>
+	T* GetInterface()
+	{
+		return static_cast<T*>(GetInterface(T::InterfaceId));
+	}
+	
+	template <typename T>
+    const T* GetInterface() const 
+    {
+    	return static_cast<const T*>(const_cast<FUAFAnimNodeData*>(this)->GetInterface(T::InterfaceId));
+    }
+};
+
+}
+
+#undef UE_API
diff --git a/Engine/Plugins/Experimental/UAF/UAFAnimNode/Source/UAFAnimNode/Public/AnimNode/UAFTransition.h b/Engine/Plugins/Experimental/UAF/UAFAnimNode/Source/UAFAnimNode/Public/AnimNode/UAFTransition.h
new file mode 100644
--- /dev/null
+++ b/Engine/Plugins/Experimental/UAF/UAFAnimNode/Source/UAFAnimNode/Public/AnimNode/UAFTransition.h
@@ -0,0 +1,98 @@
+// Copyright Epic Games, Inc. All Rights Reserved.
+
+#pragma once
+
+#include "AnimationRuntime.h"
+#include "AnimNextExecuteContext.h"
+#include "AnimNode/UAFAnimNode.h"
+
+#include "UAFTransition.generated.h"
+
+#define UE_API UAFANIMNODE_API
+
+namespace UE::UAF
+{
+	
+/**
+ * IUAFTransitionBlendNode
+ *
+ */
+struct IUAFTransition 
+{
+	static constexpr FAnimNodeInterfaceId InterfaceId = UE::UAF::FConstExprUID::MakeFromString(TEXT("IUAFTransition"));
+
+	virtual TSharedPtr<FUAFAnimNodeInstance> GetSource() = 0;
+	virtual TSharedPtr<FUAFAnimNodeInstance> GetTarget() = 0;
+	virtual bool IsComplete() = 0;
+};
+	
+/**
+ * FUAFTransitionData
+ *
+ */
+USTRUCT()
+struct FUAFTransitionData
+{
+	GENERATED_BODY()
+	virtual ~FUAFTransitionData() {};
+
+	virtual TSharedPtr<FUAFAnimNodeInstance> CreateInstance(FUAFAnimNodeUpdateContext& Context, TSharedPtr<FUAFAnimNodeInstance> Source, TSharedPtr<FUAFAnimNodeInstance> Target) const
+	{
+		ensureMsgf(false, TEXT("Animation Controller CreateInstance Unimplemented!!"));
+		return nullptr;
+	};
+
+};
+
+struct FUAFBasicTransitionInstance : public FUAFAnimNodeInstance, public IUAFTransition
+{
+
+	FUAFBasicTransitionInstance() = default;
+	FUAFBasicTransitionInstance(TSharedPtr<FUAFAnimNodeInstance> InSource, TSharedPtr<FUAFAnimNodeInstance> InTarget, float InDuration, EAlphaBlendOption InBlendOption)
+		: Source(InSource),Target(InTarget),Duration(InDuration),BlendOption(InBlendOption)
+	{
+		TimeRemaining = InDuration;
+	}
+
+
+	// IUAFTransition
+	virtual TSharedPtr<FUAFAnimNodeInstance> GetSource() override { return Source; }
+	virtual TSharedPtr<FUAFAnimNodeInstance> GetTarget() override { return Target; }
+	virtual bool IsComplete() override { return TimeRemaining <=0; }
+
+	// FUAFAnimNodeInstance
+	UE_API virtual void Update(FUAFAnimNodeUpdateContext& Context, float DeltaTime) override;
+	UE_API virtual void* GetInterface(FAnimNodeInterfaceId Id) override;
+
+	TSharedPtr<FUAFAnimNodeInstance> Source;
+	TSharedPtr<FUAFAnimNodeInstance> Target;
+	float Duration = 0.1f;
+	EAlphaBlendOption BlendOption = UE::Anim::DefaultBlendOption;
+
+	float TimeRemaining = 0.0f;
+};
+
+/**
+ * FBasicUAFTransitionBlendNodeData
+ *
+ */
+USTRUCT()
+struct FUAFBasicTransitionData : public FUAFTransitionData
+{
+	GENERATED_BODY()
+
+	UPROPERTY(EditAnywhere, Category = Blending)
+	float Duration = 0.1f;
+
+	UPROPERTY(EditAnywhere, Category = Blending)
+	EAlphaBlendOption BlendOption = UE::Anim::DefaultBlendOption;
+
+	virtual TSharedPtr<FUAFAnimNodeInstance> CreateInstance(FUAFAnimNodeUpdateContext& Context, TSharedPtr<FUAFAnimNodeInstance> Source, TSharedPtr<FUAFAnimNodeInstance> Target) const override
+	{
+		return MakeShared<FUAFBasicTransitionInstance>(Source,Target,Duration,BlendOption);
+	};
+};
+
+}
+
+#undef UE_API
diff --git a/Engine/Plugins/Experimental/UAF/UAFAnimNode/Source/UAFAnimNode/UAFAnimNode.Build.cs b/Engine/Plugins/Experimental/UAF/UAFAnimNode/Source/UAFAnimNode/UAFAnimNode.Build.cs
new file mode 100644
--- /dev/null
+++ b/Engine/Plugins/Experimental/UAF/UAFAnimNode/Source/UAFAnimNode/UAFAnimNode.Build.cs
@@ -0,0 +1,39 @@
+// Copyright Epic Games, Inc. All Rights Reserved.
+
+namespace UnrealBuildTool.Rules
+{
+	public class UAFAnimNode : ModuleRules
+	{
+		public UAFAnimNode(ReadOnlyTargetRules Target) : base(Target)
+		{
+			bAllowUETypesInNamespaces = true;
+			
+			PublicDependencyModuleNames.AddRange(new string[] { "UAF"});
+			PrivateDependencyModuleNames.AddRange(
+				new string[]
+				{
+					"Core",
+					"CoreUObject",
+					"RigVM",
+					"ControlRig",
+					"Engine",
+					"UAFAnimGraph",
+				}
+			);
+			
+			if (Target.bBuildEditor)
+			{
+				PrivateDependencyModuleNames.AddRange(
+					new string[]
+					{
+						"UAFUncookedOnly",
+						"UAFEditor",
+						"WorkspaceEditor",
+						"UAFAnimNodeUncookedOnly"
+					}
+				);
+			}
+			
+		}
+	}
+}
\ No newline at end of file
diff --git a/Engine/Plugins/Experimental/UAF/UAFAnimNode/Source/UAFAnimNodeEditor/Private/UAFAnimNodeEditorModule.cpp b/Engine/Plugins/Experimental/UAF/UAFAnimNode/Source/UAFAnimNodeEditor/Private/UAFAnimNodeEditorModule.cpp
new file mode 100644
--- /dev/null
+++ b/Engine/Plugins/Experimental/UAF/UAFAnimNode/Source/UAFAnimNodeEditor/Private/UAFAnimNodeEditorModule.cpp
@@ -0,0 +1,29 @@
+// Copyright Epic Games, Inc. All Rights Reserved.
+
+#include "UAFAnimNodeEditorModule.h"
+
+#include "Modules/ModuleManager.h"
+#include "IAnimNextEditorModule.h"
+#include "IWorkspaceEditorModule.h"
+#include "IWorkspaceEditor.h"
+
+#define LOCTEXT_NAMESPACE "UAFAnimNodeEditorModule"
+
+namespace UE::UAF::AnimNodeEditor
+{
+
+void FModule::StartupModule()
+{
+
+}
+
+void FModule::ShutdownModule()
+{
+
+}
+	
+}
+
+IMPLEMENT_MODULE(UE::UAF::AnimNodeEditor::FModule, UAFAnimNodeEditor);
+
+#undef LOCTEXT_NAMESPACE
diff --git a/Engine/Plugins/Experimental/UAF/UAFAnimNode/Source/UAFAnimNodeEditor/Private/UAFAnimNodeEditorModule.h b/Engine/Plugins/Experimental/UAF/UAFAnimNode/Source/UAFAnimNodeEditor/Private/UAFAnimNodeEditorModule.h
new file mode 100644
--- /dev/null
+++ b/Engine/Plugins/Experimental/UAF/UAFAnimNode/Source/UAFAnimNodeEditor/Private/UAFAnimNodeEditorModule.h
@@ -0,0 +1,18 @@
+// Copyright Epic Games, Inc. All Rights Reserved.
+
+#pragma once
+
+#include "CoreMinimal.h"
+#include "Modules/ModuleInterface.h"
+
+namespace UE::UAF::AnimNodeEditor
+{
+
+class FModule : public IModuleInterface
+{
+private:
+	virtual void StartupModule() override;
+	virtual void ShutdownModule() override;
+};
+
+}
diff --git a/Engine/Plugins/Experimental/UAF/UAFAnimNode/Source/UAFAnimNodeEditor/UAFAnimNodeEditor.Build.cs b/Engine/Plugins/Experimental/UAF/UAFAnimNode/Source/UAFAnimNodeEditor/UAFAnimNodeEditor.Build.cs
new file mode 100644
--- /dev/null
+++ b/Engine/Plugins/Experimental/UAF/UAFAnimNode/Source/UAFAnimNodeEditor/UAFAnimNodeEditor.Build.cs
@@ -0,0 +1,59 @@
+// Copyright Epic Games, Inc. All Rights Reserved.
+
+namespace UnrealBuildTool.Rules
+{
+	public class UAFAnimNodeEditor : ModuleRules
+	{
+		public UAFAnimNodeEditor(ReadOnlyTargetRules Target) : base(Target)
+		{
+			PublicDependencyModuleNames.AddRange(
+				new string[]
+				{
+					"Core",
+				}
+			);
+
+			PrivateDependencyModuleNames.AddRange(
+				new string[]
+				{
+					"CoreUObject",
+					"Engine",
+					"UnrealEd",
+					"InputCore",
+					"SlateCore",
+					"Slate",
+					"UAF", 
+					"UAFEditor", 
+					"UAFUncookedOnly",
+					"UAFAnimGraph",
+					"UAFAnimGraphUncookedOnly",
+					"UAFAnimNode",
+					"Settings",
+					"StructUtilsEditor",
+					"WorkspaceEditor", 
+					"EditorWidgets", 
+					"MessageLog", 
+					"AssetDefinition",
+					"RigVM",
+					"RigVMDeveloper",
+					"ToolMenus", 
+					"RigVMEditor",
+					"GraphEditor",
+					"Persona",
+					"PropertyEditor",
+					"TraceAnalysis",
+					"TraceLog",
+					"TraceServices",
+					"TraceInsights",
+					"RewindDebuggerInterface",
+					"GameplayInsights",
+					"BlueprintGraph",
+					"ToolWidgets",
+					"EditorFramework",
+					"Kismet",
+					"UAFAnimNodeUncookedOnly",
+				}
+			);
+		}
+	}
+}
\ No newline at end of file
diff --git a/Engine/Plugins/Experimental/UAF/UAFAnimNode/Source/UAFAnimNodeUncookedOnly/UAFAnimNodeUncookedOnly.Build.cs b/Engine/Plugins/Experimental/UAF/UAFAnimNode/Source/UAFAnimNodeUncookedOnly/UAFAnimNodeUncookedOnly.Build.cs
new file mode 100644
--- /dev/null
+++ b/Engine/Plugins/Experimental/UAF/UAFAnimNode/Source/UAFAnimNodeUncookedOnly/UAFAnimNodeUncookedOnly.Build.cs
@@ -0,0 +1,31 @@
+// Copyright Epic Games, Inc. All Rights Reserved.
+
+namespace UnrealBuildTool.Rules
+{
+	public class UAFAnimNodeUncookedOnly : ModuleRules
+	{
+		public UAFAnimNodeUncookedOnly(ReadOnlyTargetRules Target) : base(Target)
+		{
+			PublicDependencyModuleNames.AddRange(
+				new string[]
+				{
+					"Core",
+				}
+			);
+
+			PrivateDependencyModuleNames.AddRange(
+				new string[]
+				{
+					"CoreUObject",
+					"Engine",
+					"UAF", 
+					"UAFEditor", 
+					"UAFUncookedOnly",
+					"UAFAnimGraph",
+					"UAFAnimGraphUncookedOnly",
+					"WorkspaceEditor"
+				}
+			);
+		}
+	}
+}
\ No newline at end of file
diff --git a/Engine/Plugins/Experimental/UAF/UAFAnimNode/UAFAnimNode.uplugin b/Engine/Plugins/Experimental/UAF/UAFAnimNode/UAFAnimNode.uplugin
new file mode 100644
--- /dev/null
+++ b/Engine/Plugins/Experimental/UAF/UAFAnimNode/UAFAnimNode.uplugin
@@ -0,0 +1,51 @@
+{
+	"FileVersion" : 3,
+	"Version" : 1,
+	"VersionName" : "0.1",
+	"FriendlyName" : "UAF Anim Node",
+	"Description" : "Nodes system for UAF.",
+	"Category" : "Animation",
+	"CreatedBy" : "Epic Games, Inc.",
+	"CreatedByURL" : "https://epicgames.com",
+	"DocsURL" : "",
+	"MarketplaceURL" : "",
+	"SupportURL" : "",
+	"EnabledByDefault" : false,
+	"CanContainContent" : false,
+	"IsBetaVersion": false,
+	"IsExperimentalVersion": true,
+	"Installed" : false,
+	"Plugins":
+	[
+		{
+			"Name": "UAF",
+			"Enabled": true
+		},
+		{
+			"Name": "UAFAnimGraph",
+			"Enabled": true
+		},
+		{
+			"Name": "ControlRig",
+			"Enabled": true
+		}
+	],
+	"Modules" :
+	[
+		{
+			"Name" : "UAFAnimNode",
+			"Type" : "Runtime",
+			"LoadingPhase" : "Default"
+		},
+		{
+			"Name" : "UAFAnimNodeEditor",
+			"Type" : "Editor",
+			"LoadingPhase" : "Default"
+		},
+		{
+			"Name" : "UAFAnimNodeUncookedOnly",
+			"Type" : "UncookedOnly",
+			"LoadingPhase" : "Default"
+		}
+	]
+}
