AnimNode base cleanup/fixes

#Synced-CL 50686245
#Synced-CL 50686245

diff --git a/Engine/Plugins/Experimental/UAF/UAFAnimNode/Source/UAFAnimNode/Internal/AnimNode/UAFSequencePlayerNode.h b/Engine/Plugins/Experimental/UAF/UAFAnimNode/Source/UAFAnimNode/Internal/AnimNode/UAFSequencePlayerNode.h
--- a/Engine/Plugins/Experimental/UAF/UAFAnimNode/Source/UAFAnimNode/Internal/AnimNode/UAFSequencePlayerNode.h
+++ b/Engine/Plugins/Experimental/UAF/UAFAnimNode/Source/UAFAnimNode/Internal/AnimNode/UAFSequencePlayerNode.h
@@ -24,23 +24,25 @@ struct FUAFSequencePlayerNodeInstance : public FUAFAnimNodeInstance, public IUAF
 	virtual void* GetInterface(FAnimNodeInterfaceId Id) override;
 
 	// IUAFAnimNodeTimeline
-	virtual float GetCurrentTime() override;
+	virtual float GetCurrentTime() const override;
 	virtual void SetCurrentTime(float InCurrentTime) override;
-	virtual bool IsLooping() override;
-	virtual float GetLength() override;
+	virtual bool IsLooping() const override;
+	virtual float GetLength() const override;
+	virtual ETypeAdvanceAnim GetTimeAdvanceResult() const override;
 
 private:
 	/** Delta time range required for root motion extraction **/
 	FDeltaTimeRecord DeltaTimeRecord;
 	float CurrentTime = 0.0f;
 
-	TObjectPtr<const UAnimSequence> Sequence;
+	const UAnimSequence* Sequence = nullptr;
+	ETypeAdvanceAnim TimeAdvanceResult;
 	bool bIsLooping;
 };
 
 
 // AnimNode Data for playing an AnimSequence
-USTRUCT()
+USTRUCT(DisplayName = "Sequence Player")
 struct FUAFSequencePlayerNodeData : public FUAFAnimNodeData
 #if CPP
 , public IUAFAnimNodeDataHasAsset
diff --git a/Engine/Plugins/Experimental/UAF/UAFAnimNode/Source/UAFAnimNode/Private/UAFSequencePlayerNode.cpp b/Engine/Plugins/Experimental/UAF/UAFAnimNode/Source/UAFAnimNode/Private/UAFSequencePlayerNode.cpp
--- a/Engine/Plugins/Experimental/UAF/UAFAnimNode/Source/UAFAnimNode/Private/UAFSequencePlayerNode.cpp
+++ b/Engine/Plugins/Experimental/UAF/UAFAnimNode/Source/UAFAnimNode/Private/UAFSequencePlayerNode.cpp
@@ -28,7 +28,7 @@ void FUAFSequencePlayerNodeInstance::Update(FUAFAnimNodeUpdateContext& Context,
 		const float SequenceLength = Sequence->GetPlayLength();
 
 		DeltaTimeRecord.Set(CurrentTime, DeltaTime);
-		FAnimationRuntime::AdvanceTime(bIsLooping, DeltaTime, CurrentTime, SequenceLength);
+		TimeAdvanceResult = FAnimationRuntime::AdvanceTime(bIsLooping, DeltaTime, CurrentTime, SequenceLength);
 	
 		FAnimNextAnimSequenceKeyframeTask Task;
 		Task.AnimSequence = Sequence;
@@ -57,18 +57,23 @@ void FUAFSequencePlayerNodeInstance::SetCurrentTime(float InCurrentTime)
 	CurrentTime = InCurrentTime;
 }
 
-bool FUAFSequencePlayerNodeInstance::IsLooping()
+ETypeAdvanceAnim FUAFSequencePlayerNodeInstance::GetTimeAdvanceResult() const
+{
+	return TimeAdvanceResult;
+}
+
+bool FUAFSequencePlayerNodeInstance::IsLooping() const
 {
 	return bIsLooping;
 }
 
 
-float FUAFSequencePlayerNodeInstance::GetCurrentTime()
+float FUAFSequencePlayerNodeInstance::GetCurrentTime() const
 {
 	return CurrentTime;
 }
 
-float FUAFSequencePlayerNodeInstance::GetLength()
+float FUAFSequencePlayerNodeInstance::GetLength() const
 {
 	if (Sequence)
 	{
diff --git a/Engine/Plugins/Experimental/UAF/UAFAnimNode/Source/UAFAnimNode/Public/AnimNode/UAFAnimNode.h b/Engine/Plugins/Experimental/UAF/UAFAnimNode/Source/UAFAnimNode/Public/AnimNode/UAFAnimNode.h
--- a/Engine/Plugins/Experimental/UAF/UAFAnimNode/Source/UAFAnimNode/Public/AnimNode/UAFAnimNode.h
+++ b/Engine/Plugins/Experimental/UAF/UAFAnimNode/Source/UAFAnimNode/Public/AnimNode/UAFAnimNode.h
@@ -19,10 +19,11 @@ namespace UE::UAF
 struct IUAFAnimNodeTimeline
 {
 	static constexpr FAnimNodeInterfaceId InterfaceId = FConstExprUID::MakeFromString(TEXT("IUAFAnimNodeTimeline"));
-	virtual float GetCurrentTime() = 0;
-	virtual float GetLength() = 0;
-	virtual bool IsLooping() = 0;
+	virtual float GetCurrentTime() const = 0;
+	virtual float GetLength() const = 0;
+	virtual bool IsLooping() const = 0;
 	virtual void SetCurrentTime(float CurrentTime) = 0;
+	virtual ETypeAdvanceAnim GetTimeAdvanceResult() const = 0;
 };
 
 
diff --git a/Engine/Plugins/Experimental/UAF/UAFAnimNode/Source/UAFAnimNode/Public/AnimNode/UAFTransition.h b/Engine/Plugins/Experimental/UAF/UAFAnimNode/Source/UAFAnimNode/Public/AnimNode/UAFTransition.h
--- a/Engine/Plugins/Experimental/UAF/UAFAnimNode/Source/UAFAnimNode/Public/AnimNode/UAFTransition.h
+++ b/Engine/Plugins/Experimental/UAF/UAFAnimNode/Source/UAFAnimNode/Public/AnimNode/UAFTransition.h
@@ -38,7 +38,7 @@ struct FUAFTransitionData
 
 	virtual TSharedPtr<FUAFAnimNodeInstance> CreateInstance(FUAFAnimNodeUpdateContext& Context, TSharedPtr<FUAFAnimNodeInstance> Source, TSharedPtr<FUAFAnimNodeInstance> Target) const
 	{
-		ensureMsgf(false, TEXT("Animation Controller CreateInstance Unimplemented!!"));
+		ensureMsgf(false, TEXT("FUAFTransition CreateInstance Unimplemented!!"));
 		return nullptr;
 	};
 
@@ -76,7 +76,7 @@ struct FUAFBasicTransitionInstance : public FUAFAnimNodeInstance, public IUAFTra
  * FBasicUAFTransitionBlendNodeData
  *
  */
-USTRUCT()
+USTRUCT(DisplayName = "Basic Transition")
 struct FUAFBasicTransitionData : public FUAFTransitionData
 {
 	GENERATED_BODY()
@@ -89,6 +89,10 @@ struct FUAFBasicTransitionData : public FUAFTransitionData
 
 	virtual TSharedPtr<FUAFAnimNodeInstance> CreateInstance(FUAFAnimNodeUpdateContext& Context, TSharedPtr<FUAFAnimNodeInstance> Source, TSharedPtr<FUAFAnimNodeInstance> Target) const override
 	{
+		if (Duration <= 0)
+		{
+			return Target;			
+		}
 		return MakeShared<FUAFBasicTransitionInstance>(Source,Target,Duration,BlendOption);
 	};
 };
diff --git a/Engine/Plugins/Experimental/UAF/UAFAnimNode/Source/UAFAnimNodeUncookedOnly/Private/UAFAnimNodeUncookedOnlyModule.cpp b/Engine/Plugins/Experimental/UAF/UAFAnimNode/Source/UAFAnimNodeUncookedOnly/Private/UAFAnimNodeUncookedOnlyModule.cpp
new file mode 100644
--- /dev/null
+++ b/Engine/Plugins/Experimental/UAF/UAFAnimNode/Source/UAFAnimNodeUncookedOnly/Private/UAFAnimNodeUncookedOnlyModule.cpp
@@ -0,0 +1,6 @@
+﻿// Copyright Epic Games, Inc. All Rights Reserved.
+
+#include "CoreMinimal.h"
+#include "Modules/ModuleManager.h"
+
+IMPLEMENT_MODULE(FDefaultModuleImpl, UAFAnimNodeUncookedOnly);
\ No newline at end of file
