Export missing symbols

#Synced-CL 51273531

diff --git a/Engine/Plugins/Experimental/UAF/UAFAnimNode/Source/UAFAnimNode/Public/UAF/AnimNodes/UAFSimpleTransition.h b/Engine/Plugins/Experimental/UAF/UAFAnimNode/Source/UAFAnimNode/Public/UAF/AnimNodes/UAFSimpleTransition.h
--- a/Engine/Plugins/Experimental/UAF/UAFAnimNode/Source/UAFAnimNode/Public/UAF/AnimNodes/UAFSimpleTransition.h
+++ b/Engine/Plugins/Experimental/UAF/UAFAnimNode/Source/UAFAnimNode/Public/UAF/AnimNodes/UAFSimpleTransition.h
@@ -52,8 +52,8 @@ namespace UE::UAF
 		virtual bool IsComplete() const override;
 
 		// FUAFAnimNode impl
-		virtual void PreUpdate(FUAFAnimGraphUpdateContext& GraphContext) override;
-		virtual void* GetInterface(FUAFAnimNodeInterfaceId Id) override;
+		UE_API virtual void PreUpdate(FUAFAnimGraphUpdateContext& GraphContext) override;
+		UE_API virtual void* GetInterface(FUAFAnimNodeInterfaceId Id) override;
 
 	protected:
 		static constexpr int32 SourceChildIndex = 0;
