Ensure leaf nodes produce an AnimOp

#Synced-CL 51372242

diff --git a/Engine/Plugins/Experimental/UAF/UAFAnimNode/Source/UAFAnimNode/Private/UAF/AnimNodeCore/UAFAnimNodeUpdate.cpp b/Engine/Plugins/Experimental/UAF/UAFAnimNode/Source/UAFAnimNode/Private/UAF/AnimNodeCore/UAFAnimNodeUpdate.cpp
--- a/Engine/Plugins/Experimental/UAF/UAFAnimNode/Source/UAFAnimNode/Private/UAF/AnimNodeCore/UAFAnimNodeUpdate.cpp
+++ b/Engine/Plugins/Experimental/UAF/UAFAnimNode/Source/UAFAnimNode/Private/UAF/AnimNodeCore/UAFAnimNodeUpdate.cpp
@@ -187,6 +187,7 @@ namespace UE::UAF
 #if DO_CHECK
 				CurrentNode->bIsInPostUpdate = true;
 #endif
+
 				// We've already visited this node once, time to post-update
 				CurrentNode->PostUpdate(UpdateContext);
 
@@ -200,6 +201,12 @@ namespace UE::UAF
 				{
 					EvaluationProgram.AppendTaskPtr(AnimOp);
 				}
+#if DO_CHECK
+				else
+				{
+					checkf(Entry->Node->GetNumChildren() != 0, TEXT("Leaf nodes without children must have a PostAnimOp to produce an output."));
+				}
+#endif
 
 				// Now that we have been visited, we are no longer newly relevant
 				CurrentNode->SetIsNewlyRelevant(false);
diff --git a/Engine/Plugins/Experimental/UAF/UAFAnimNode/Source/UAFAnimNode/Public/UAF/AnimNodeCore/UAFAnimNode.h b/Engine/Plugins/Experimental/UAF/UAFAnimNode/Source/UAFAnimNode/Public/UAF/AnimNodeCore/UAFAnimNode.h
--- a/Engine/Plugins/Experimental/UAF/UAFAnimNode/Source/UAFAnimNode/Public/UAF/AnimNodeCore/UAFAnimNode.h
+++ b/Engine/Plugins/Experimental/UAF/UAFAnimNode/Source/UAFAnimNode/Public/UAF/AnimNodeCore/UAFAnimNode.h
@@ -91,11 +91,14 @@ namespace UE::UAF
 		// Main update entry point for node instances
 		// This function is responsible for updating topology (e.g. managing children) and advancing time/phase
 		// Called during the update pass before children have been queued and traversed
+		// When choosing whether to use PreUpdate or PostUpdate, favor the former when possible
 		UE_API virtual void PreUpdate(FUAFAnimGraphUpdateContext& GraphContext);
 
 		// Main update entry point for node instances
 		// This function is responsible for updating topology (e.g. managing children) and advancing time/phase
 		// Called during the update pass after children have been queued and traversed
+		// Note that it is not possible to change the PreAnimOp while in PostUpdate as it has already been queued
+		// When choosing whether to use PreUpdate or PostUpdate, favor the former when possible
 		UE_API virtual void PostUpdate(FUAFAnimGraphUpdateContext& GraphContext);
 
 		// Adds a new child to this node and returns its index
@@ -113,9 +116,12 @@ namespace UE::UAF
 		UE_API void SetChildAt(const FUAFAnimNodePtr& Child, int32 ChildIndex);
 
 		// Sets the AnimOp to perform before our children evaluate theirs
+		// When choosing whether to use PreAnimOp or PostAnimOp, favor the latter when possible
 		void SetPreAnimOp(TSharedPtr<FAnimNextEvaluationTask> AnimOp);
 
 		// Sets the AnimOp to perform after our children evaluate theirs
+		// When choosing whether to use PreAnimOp or PostAnimOp, favor the latter when possible
+		// Must be used by leaf nodes to produce an output
 		void SetPostAnimOp(TSharedPtr<FAnimNextEvaluationTask> AnimOp);
 
 	private:
