Fix compile time function implementation check to work with static function

#Synced-CL 51372242

diff --git a/Engine/Plugins/Experimental/UAF/UAFAnimNode/Source/UAFAnimNode/Public/UAF/AnimNodeCore/UAFAnimNode.h b/Engine/Plugins/Experimental/UAF/UAFAnimNode/Source/UAFAnimNode/Public/UAF/AnimNodeCore/UAFAnimNode.h
--- a/Engine/Plugins/Experimental/UAF/UAFAnimNode/Source/UAFAnimNode/Public/UAF/AnimNodeCore/UAFAnimNode.h
+++ b/Engine/Plugins/Experimental/UAF/UAFAnimNode/Source/UAFAnimNode/Public/UAF/AnimNodeCore/UAFAnimNode.h
@@ -196,7 +196,8 @@ namespace UE::UAF
 	template<class AnimNodeType>
 	inline void FUAFAnimNode::InitializeAs(FUAFAnimGraphUpdateContext& UpdateContext)
 	{
-		bHasAddReferencedObjects = !std::is_same_v<decltype(&AnimNodeType::AddReferencedObjects), decltype(&FUAFAnimNode::AddReferencedObjects)>;
+		// Cannot use std::same_v with static functions since their signature is the same
+		bHasAddReferencedObjects = std::conditional_t<&AnimNodeType::AddReferencedObjects != &FUAFAnimNode::AddReferencedObjects, std::true_type, std::false_type>::value;
 
 #if DO_CHECK
 		bIsInitialized = true;
