Simplify newly relevant tracking in new UAF runtime

#Synced-CL 51372242

diff --git a/Engine/Plugins/Experimental/UAF/UAFAnimNode/Source/UAFAnimNode/Private/UAF/AnimNodeCore/UAFAnimNode.cpp b/Engine/Plugins/Experimental/UAF/UAFAnimNode/Source/UAFAnimNode/Private/UAF/AnimNodeCore/UAFAnimNode.cpp
--- a/Engine/Plugins/Experimental/UAF/UAFAnimNode/Source/UAFAnimNode/Private/UAF/AnimNodeCore/UAFAnimNode.cpp
+++ b/Engine/Plugins/Experimental/UAF/UAFAnimNode/Source/UAFAnimNode/Private/UAF/AnimNodeCore/UAFAnimNode.cpp
@@ -19,6 +19,27 @@ namespace UE::UAF
 		}
 	}
 
+	void FUAFAnimNode::Reset()
+	{
+		checkf(Parent == nullptr, TEXT("Attempting to reset a node that has a parent."));
+
+		// Detach our children and reset everything
+		for (const FUAFAnimNodePtr& Child : Children)
+		{
+			if (Child)
+			{
+				Child->Parent = nullptr;
+			}
+		}
+		Children.Reset();
+
+		PreAnimOp = nullptr;
+		PostAnimOp = nullptr;
+		TotalWeight = 0.0f;
+		bIsBlendingOut = false;
+		bIsNewlyRelevant = true;
+	}
+
 	int32 FUAFAnimNode::AddChild(FUAFAnimNodePtr&& Child)
 	{
 		checkf(!Child || !Children.Contains(Child), TEXT("Attempting to add a duplicate child entry."));
diff --git a/Engine/Plugins/Experimental/UAF/UAFAnimNode/Source/UAFAnimNode/Private/UAF/AnimNodeCore/UAFAnimNodeUpdate.cpp b/Engine/Plugins/Experimental/UAF/UAFAnimNode/Source/UAFAnimNode/Private/UAF/AnimNodeCore/UAFAnimNodeUpdate.cpp
--- a/Engine/Plugins/Experimental/UAF/UAFAnimNode/Source/UAFAnimNode/Private/UAF/AnimNodeCore/UAFAnimNodeUpdate.cpp
+++ b/Engine/Plugins/Experimental/UAF/UAFAnimNode/Source/UAFAnimNode/Private/UAF/AnimNodeCore/UAFAnimNodeUpdate.cpp
@@ -172,9 +172,8 @@ namespace UE::UAF
 						Child->SetTotalWeight(CurrentNode->GetTotalWeight());
 					}
 
-					// Always inherit whether we are blending out or newly relevant
+					// Always inherit whether we are blending out
 					Child->SetIsBlendingOut(CurrentNode->IsBlendingOut());
-					Child->SetIsNewlyRelevant(CurrentNode->IsNewlyRelevant());
 
 					NodesPendingUpdateStackTop = ChildEntry;
 				}
@@ -209,7 +208,7 @@ namespace UE::UAF
 #endif
 
 				// Now that we have been visited, we are no longer newly relevant
-				CurrentNode->SetIsNewlyRelevant(false);
+				CurrentNode->bIsNewlyRelevant = false;
 
 				// Now that we are done processing this entry, we can pop it
 				NodesPendingUpdateStackTop = Entry->PrevStackEntry;
@@ -224,7 +223,7 @@ namespace UE::UAF
 		}
 
 		// Now that the root has updated, it is no longer newly relevant
-		RootNode.SetIsNewlyRelevant(false);
+		RootNode.bIsNewlyRelevant = false;
 
 		return EvaluationProgram;
 	}
diff --git a/Engine/Plugins/Experimental/UAF/UAFAnimNode/Source/UAFAnimNode/Public/UAF/AnimNodeCore/UAFAnimNode.h b/Engine/Plugins/Experimental/UAF/UAFAnimNode/Source/UAFAnimNode/Public/UAF/AnimNodeCore/UAFAnimNode.h
--- a/Engine/Plugins/Experimental/UAF/UAFAnimNode/Source/UAFAnimNode/Public/UAF/AnimNodeCore/UAFAnimNode.h
+++ b/Engine/Plugins/Experimental/UAF/UAFAnimNode/Source/UAFAnimNode/Public/UAF/AnimNodeCore/UAFAnimNode.h
@@ -32,6 +32,11 @@ namespace UE::UAF
 		FUAFAnimNode() = delete;
 		FUAFAnimNode(FUAFAnimGraphUpdateContext& UpdateContext);
 
+		// Resets this node to its initial state
+		// This must be called when a node is re-used
+		// Cannot be called if the node is attached to a parent
+		UE_API void Reset();
+
 		// Returns the child at the specified index
 		const FUAFAnimNodePtr& GetChildAt(int32 ChildIndex) const;
 
@@ -60,11 +65,10 @@ namespace UE::UAF
 		void SetIsBlendingOut(bool bIsBlendingOut);
 
 		// Returns whether or not we are newly relevant
+		// A node is newly relevant the first time it updates
+		// If you wish to re-use a node after it has fully blended out, you must first reset it
 		bool IsNewlyRelevant() const;
 
-		// Sets whether or not we are newly relevant
-		void SetIsNewlyRelevant(bool bIsNewlyRelevant);
-
 		// Returns the AnimOp to perform before our children evaluate theirs
 		const TSharedPtr<FAnimNextEvaluationTask>& GetPreAnimOp() const;
 
@@ -263,11 +267,6 @@ namespace UE::UAF
 		return bIsNewlyRelevant;
 	}
 
-	inline void FUAFAnimNode::SetIsNewlyRelevant(bool bInIsNewlyRelevant)
-	{
-		bIsNewlyRelevant = bInIsNewlyRelevant;
-	}
-
 	inline const TSharedPtr<FAnimNextEvaluationTask>& FUAFAnimNode::GetPreAnimOp() const
 	{
 		return PreAnimOp;
