Avoid using the output skeletal mesh when none is present in new UAF runtime

#Synced-CL 51626790

diff --git a/Engine/Plugins/Experimental/UAF/UAFAnimNode/Source/UAFAnimNode/Private/RigUnit_RunAnimNode.cpp b/Engine/Plugins/Experimental/UAF/UAFAnimNode/Source/UAFAnimNode/Private/RigUnit_RunAnimNode.cpp
--- a/Engine/Plugins/Experimental/UAF/UAFAnimNode/Source/UAFAnimNode/Private/RigUnit_RunAnimNode.cpp
+++ b/Engine/Plugins/Experimental/UAF/UAFAnimNode/Source/UAFAnimNode/Private/RigUnit_RunAnimNode.cpp
@@ -95,19 +95,23 @@ FRigUnit_RunAnimNode_Execute()
 		{
 			FUAFSystemOutputComponent& OutputPoseComponent = ModuleInstance.GetComponent<FUAFSystemOutputComponent>();
 
+			// Note that the skeletal mesh and skeleton can be null
 			USkeletalMesh* SkeletalMesh = OutputPoseComponent.GetSkeletalMesh();
-			check(SkeletalMesh);
+			USkeleton* Skeleton = SkeletalMesh != nullptr ? SkeletalMesh->GetSkeleton() : nullptr;
 
-			USkeleton* Skeleton = SkeletalMesh->GetSkeleton();
-			check(Skeleton);
+			if (Skeleton != nullptr)
+			{
+				TNonNullPtr<UUAFSkeletonUserData> UserData = UUAFSkeletonUserData::FromSkeleton(Skeleton);
+				UAbstractSkeletonSetBinding* Binding = UserData->GetSetBinding();
 
-			TNonNullPtr<UUAFSkeletonUserData> UserData = UUAFSkeletonUserData::FromSkeleton(Skeleton);
+				FUAFEvaluateValuesArgs EvaluateArgs(*Binding, AttributeSet, SkeletalMesh, DesiredLOD, AnimOpList, RefPose);
 
-			UAbstractSkeletonSetBinding* Binding = UserData->GetSetBinding();
-
-			FUAFEvaluateValuesArgs EvaluateArgs(*Binding, AttributeSet, SkeletalMesh, DesiredLOD, AnimOpList, RefPose);
-
-			EvaluateValues(EvaluateArgs, Result);
+				EvaluateValues(EvaluateArgs, Result);
+			}
+			else
+			{
+				// If we have no skeletal mesh or skeleton, we output nothing
+			}
 		}
 
 		// Now that everything has updated and we've computed our values, we know for certain every blend weight
