Implement play rate scaling during UAF anim node update

#Synced-CL 51626790

diff --git a/Engine/Plugins/Experimental/UAF/UAFAnimNode/Source/UAFAnimNode/Private/UAF/AnimNodeCore/UAFAnimNodeUpdate.cpp b/Engine/Plugins/Experimental/UAF/UAFAnimNode/Source/UAFAnimNode/Private/UAF/AnimNodeCore/UAFAnimNodeUpdate.cpp
--- a/Engine/Plugins/Experimental/UAF/UAFAnimNode/Source/UAFAnimNode/Private/UAF/AnimNodeCore/UAFAnimNodeUpdate.cpp
+++ b/Engine/Plugins/Experimental/UAF/UAFAnimNode/Source/UAFAnimNode/Private/UAF/AnimNodeCore/UAFAnimNodeUpdate.cpp
@@ -65,8 +65,9 @@ namespace UE::UAF
 		: PreviousContext(Private::GCurrentUpdateContext)
 		, HostObject(InHostObject)
 		, VariablesOwner(InVariablesOwner)
-		, GCReferences(InGCReferences)
 		, DeltaTime(InDeltaTime)
+		, UpdateDeltaTime(InDeltaTime)
+		, GCReferences(InGCReferences)
 	{
 		Private::GCurrentUpdateContext = this;
 
@@ -83,6 +84,8 @@ namespace UE::UAF
 
 	FUAFAnimGraphUpdateContext::~FUAFAnimGraphUpdateContext()
 	{
+		checkf(PreviousPlayRateStack.IsEmpty(), TEXT("A play rate was pushed and not popped."));
+
 		// Destroy any nodes that are pending destruction
 		// Note that this may queue up more nodes for destruction to avoid recursion
 		while (!NodesPendingDestruction.IsEmpty())
@@ -104,6 +107,29 @@ namespace UE::UAF
 		return Private::GCurrentUpdateContext;
 	}
 
+	void FUAFAnimGraphUpdateContext::PushPlayRate(float PlayRate)
+	{
+		// Save our current play rate
+		PreviousPlayRateStack.Push(CurrentPlayRate);
+
+		// Combine the current play rate with the new one
+		CurrentPlayRate *= PlayRate;
+
+		// Update our delta time based on the new play rate
+		DeltaTime = UpdateDeltaTime * CurrentPlayRate;
+	}
+
+	void FUAFAnimGraphUpdateContext::PopPlayRate()
+	{
+		checkf(!PreviousPlayRateStack.IsEmpty(), TEXT("Attempting to pop a play rate when none was pushed."));
+
+		// Restore our previous play rate
+		CurrentPlayRate = PreviousPlayRateStack.Pop(EAllowShrinking::No);
+
+		// Update our delta time based on the new play rate
+		DeltaTime = UpdateDeltaTime * CurrentPlayRate;
+	}
+
 	void FUAFAnimGraphUpdateContext::QueueForDestruction(FUAFAnimNode* AnimNode)
 	{
 		NodesPendingDestruction.Push(AnimNode);
diff --git a/Engine/Plugins/Experimental/UAF/UAFAnimNode/Source/UAFAnimNode/Private/UAF/AnimNodes/UAFScalePlayRate.cpp b/Engine/Plugins/Experimental/UAF/UAFAnimNode/Source/UAFAnimNode/Private/UAF/AnimNodes/UAFScalePlayRate.cpp
new file mode 100644
--- /dev/null
+++ b/Engine/Plugins/Experimental/UAF/UAFAnimNode/Source/UAFAnimNode/Private/UAF/AnimNodes/UAFScalePlayRate.cpp
@@ -0,0 +1,46 @@
+// Copyright Epic Games, Inc. All Rights Reserved.
+
+#include "UAF/AnimNodes/UAFScalePlayRate.h"
+
+#include "UAF/AnimNodeCore/UAFAnimNodeUpdate.h"
+#include "UAF/AnimOps/UAFNullAnimOp.h"
+
+#include UE_INLINE_GENERATED_CPP_BY_NAME(UAFScalePlayRate)
+
+namespace UE::UAF
+{
+	FUAFScalePlayRate::FUAFScalePlayRate(FUAFAnimGraphUpdateContext& Context, const FUAFScalePlayRateData& Data)
+		: FUAFAnimNode(Context)
+		, PlayRateMultiplier(Data.PlayRateMultiplier)
+	{
+		InitializeAs<FUAFScalePlayRate>(Context);
+
+		if (const FUAFAnimNodeData* ChildData = Data.Child.GetPtr())
+		{
+			if (FUAFAnimNodePtr Child = ChildData->CreateInstance(Context))
+			{
+				AddChild(MoveTemp(Child));
+			}
+		}
+
+		if (!HasChildren())
+		{
+			SetPostAnimOp(FUAFNullAnimOp::Get());
+		}
+	}
+
+	void FUAFScalePlayRate::PreUpdate(FUAFAnimGraphUpdateContext& Context)
+	{
+		Context.PushPlayRate(PlayRateMultiplier);
+	}
+
+	void FUAFScalePlayRate::PostUpdate(FUAFAnimGraphUpdateContext& Context)
+	{
+		Context.PopPlayRate();
+	}
+
+	FUAFAnimNodePtr FUAFScalePlayRateData::CreateInstance(FUAFAnimGraphUpdateContext& Context) const
+	{
+		return MakeAnimNode<FUAFScalePlayRate>(Context, *this);
+	}
+}
diff --git a/Engine/Plugins/Experimental/UAF/UAFAnimNode/Source/UAFAnimNode/Public/UAF/AnimNodeCore/UAFAnimNodeUpdate.h b/Engine/Plugins/Experimental/UAF/UAFAnimNode/Source/UAFAnimNode/Public/UAF/AnimNodeCore/UAFAnimNodeUpdate.h
--- a/Engine/Plugins/Experimental/UAF/UAFAnimNode/Source/UAFAnimNode/Public/UAF/AnimNodeCore/UAFAnimNodeUpdate.h
+++ b/Engine/Plugins/Experimental/UAF/UAFAnimNode/Source/UAFAnimNode/Public/UAF/AnimNodeCore/UAFAnimNodeUpdate.h
@@ -53,9 +53,21 @@ namespace UE::UAF
 		// Returns a reference collector for anim nodes that need exposure to GC
 		[[nodiscard]] FUAFAnimNodeReferenceCollector& GetGCReferences();
 
-		// Returns how much time should the update pass advance by
+		// Returns how much time should the update pass advance by factoring in the current play rate scaling
 		[[nodiscard]] float GetDeltaTime() const;
 
+		// Returns how much time should the update pass advance by
+		[[nodiscard]] float GetUpdateDeltaTime() const;
+
+		// Returns the currently applied play rate when computing the delta time
+		[[nodiscard]] float GetPlayRate() const;
+
+		// Pushes a new play rate scale on top of the stack
+		void PushPlayRate(float PlayRate);
+
+		// Pops the most recent play rate scale from the top of the stack
+		void PopPlayRate();
+
 #if DO_CHECK
 		// Debug counter used to enforce invariants
 		[[nodiscard]] uint32 GetUpdateCounter() const;
@@ -69,12 +81,22 @@ namespace UE::UAF
 		UObject* HostObject = nullptr;
 		FUAFAssetInstance* VariablesOwner = nullptr;
 
-		FUAFAnimNodeReferenceCollector& GCReferences;
+		// A stack of previous play rates
+		TArray<float> PreviousPlayRateStack;
 
-		TArray<FUAFAnimNode*> NodesPendingDestruction;
+		// The currently active play rate
+		float CurrentPlayRate = 1.0f;
 
+		// The delta time as currently scaled by the play rate
 		float DeltaTime = 0.0f;
 
+		// The unscaled delta time used by the update pass
+		float UpdateDeltaTime = 0.0f;
+
+		FUAFAnimNodeReferenceCollector& GCReferences;
+
+		TArray<FUAFAnimNode*> NodesPendingDestruction;
+
 #if DO_CHECK
 		// Debug counter used to enforce invariants
 		uint32 UpdateCounter = 0;
@@ -169,6 +191,16 @@ namespace UE::UAF
 		return DeltaTime;
 	}
 
+	inline float FUAFAnimGraphUpdateContext::GetUpdateDeltaTime() const
+	{
+		return UpdateDeltaTime;
+	}
+
+	inline float FUAFAnimGraphUpdateContext::GetPlayRate() const
+	{
+		return CurrentPlayRate;
+	}
+
 #if DO_CHECK
 	inline uint32 FUAFAnimGraphUpdateContext::GetUpdateCounter() const
 	{
diff --git a/Engine/Plugins/Experimental/UAF/UAFAnimNode/Source/UAFAnimNode/Public/UAF/AnimNodes/UAFScalePlayRate.h b/Engine/Plugins/Experimental/UAF/UAFAnimNode/Source/UAFAnimNode/Public/UAF/AnimNodes/UAFScalePlayRate.h
new file mode 100644
--- /dev/null
+++ b/Engine/Plugins/Experimental/UAF/UAFAnimNode/Source/UAFAnimNode/Public/UAF/AnimNodes/UAFScalePlayRate.h
@@ -0,0 +1,45 @@
+// Copyright Epic Games, Inc. All Rights Reserved.
+
+#pragma once
+
+#include "CoreMinimal.h"
+#include "UAF/AnimNodeCore/UAFAnimNodeData.h"
+#include "UAF/AnimNodeCore/UAFAnimNode.h"
+
+#include "UAFScalePlayRate.generated.h"
+
+#define UE_API UAFANIMNODE_API
+
+namespace UE::UAF
+{
+	// Anim node shared data for scaling the play rate of a child
+	USTRUCT(DisplayName = "Scale Play Rate")
+	struct FUAFScalePlayRateData : public FUAFAnimNodeData
+	{
+		GENERATED_BODY()
+
+		UPROPERTY(EditAnywhere, Category = "Data")
+		TInstancedStruct<FUAFAnimNodeData> Child;
+
+		UPROPERTY(EditAnywhere, Category = "Data")
+		float PlayRateMultiplier = 1.0f;
+
+		// FUAFAnimNodeData impl
+		UE_API virtual FUAFAnimNodePtr CreateInstance(FUAFAnimGraphUpdateContext& Context) const override;
+	};
+
+	// Anim node instance for scaling the play rate of a child
+	struct FUAFScalePlayRate : public FUAFAnimNode
+	{
+		UE_API FUAFScalePlayRate(FUAFAnimGraphUpdateContext& Context, const FUAFScalePlayRateData& InData);
+
+		// FUAFAnimNode impl
+		UE_API virtual void PreUpdate(FUAFAnimGraphUpdateContext& Context) override;
+		UE_API virtual void PostUpdate(FUAFAnimGraphUpdateContext& Context) override;
+
+	private:
+		float PlayRateMultiplier = 1.0f;
+	};
+}
+
+#undef UE_API
