Blend By [REDACTED_PERSON]

#Synced-CL 51640609

diff --git a/Engine/Plugins/Experimental/UAF/UAFAnimNode/Source/UAFAnimNode/Private/UAF/AnimNodes/UAFBlendByBoolNode.cpp b/Engine/Plugins/Experimental/UAF/UAFAnimNode/Source/UAFAnimNode/Private/UAF/AnimNodes/UAFBlendByBoolNode.cpp
new file mode 100644
--- /dev/null
+++ b/Engine/Plugins/Experimental/UAF/UAFAnimNode/Source/UAFAnimNode/Private/UAF/AnimNodes/UAFBlendByBoolNode.cpp
@@ -0,0 +1,70 @@
+// Copyright Epic Games, Inc. All Rights Reserved.
+#include "UAF/AnimNodes/UAFBlendByBoolNode.h"
+#include "Module/AnimNextModuleInstance.h"
+#include "UAF/AnimNodeCore/UAFAnimNodeUpdate.h"
+
+#include UE_INLINE_GENERATED_CPP_BY_NAME(UAFBlendByBoolNode)
+
+namespace UE::UAF
+{
+	FUAFBlendByBoolNode::FUAFBlendByBoolNode(FUAFAnimGraphUpdateContext& Context, const FUAFBlendByBoolNodeData* InData)
+		: FUAFBlendStack(Context), Data(InData)
+	{
+		InitializeAs<FUAFBlendByBoolNode>(Context);
+
+		Context.GetVariablesOwner()->GetVariable(Data->Bool, bCurrentValue);
+
+		if (bCurrentValue)
+		{
+			TransitionTo(Context, Data->TrueNode.Get().CreateInstance(Context));
+		}
+		else
+		{
+			TransitionTo(Context, Data->FalseNode.Get().CreateInstance(Context));
+		}
+	}
+
+	void FUAFBlendByBoolNode::PreUpdate(FUAFAnimGraphUpdateContext& Context)
+	{
+		FUAFBlendStack::PreUpdate(Context);
+
+		bool bNewValue = false;
+		Context.GetVariablesOwner()->GetVariable(Data->Bool, bNewValue);
+
+		if (bCurrentValue != bNewValue)
+		{
+			bCurrentValue = bNewValue;
+			if (bCurrentValue)
+			{
+				TransitionTo(Context, Data->TrueNode.Get().CreateInstance(Context), Data->Transition.GetPtr());
+			}
+			else
+			{
+				TransitionTo(Context, Data->FalseNode.Get().CreateInstance(Context), Data->Transition.GetPtr());
+			}
+		}
+
+		ValidateSingleChild();
+	}
+	
+	
+	FUAFAnimNodePtr FUAFBlendByBoolNodeData::CreateInstance(FUAFAnimGraphUpdateContext& Context) const
+	{
+		if (EvaluationFrequency == EBlendByBoolEvaluationFrequency::OnCreate)
+		{
+			bool bValue = false;
+			Context.GetVariablesOwner()->GetVariable(Bool, bValue);
+
+			if (bValue)
+			{
+				return TrueNode.Get().CreateInstance(Context);
+			}
+			else
+			{
+				return FalseNode.Get().CreateInstance(Context);
+			}
+		}
+		
+		return MakeAnimNode<FUAFBlendByBoolNode>(Context, this);
+	}
+}
diff --git a/Engine/Plugins/Experimental/UAF/UAFAnimNode/Source/UAFAnimNode/Public/UAF/AnimNodes/UAFBlendByBoolNode.h b/Engine/Plugins/Experimental/UAF/UAFAnimNode/Source/UAFAnimNode/Public/UAF/AnimNodes/UAFBlendByBoolNode.h
new file mode 100644
--- /dev/null
+++ b/Engine/Plugins/Experimental/UAF/UAFAnimNode/Source/UAFAnimNode/Public/UAF/AnimNodes/UAFBlendByBoolNode.h
@@ -0,0 +1,71 @@
+// Copyright Epic Games, Inc. All Rights Reserved.
+
+#pragma once
+
+#include "CoreMinimal.h"
+#include "EvaluationVM/Tasks/PushAnimSequenceKeyframe.h"
+#include "UAF/AnimNodeCore/UAFAnimNodeData.h"
+#include "UAF/AnimNodes/UAFBlendStack.h"
+#include "Variables/AnimNextVariableReference.h"
+
+#include "UAFBlendByBoolNode.generated.h"
+
+#define UE_API UAFANIMNODE_API
+
+namespace UE::UAF
+{
+	UENUM()
+	enum class EBlendByBoolEvaluationFrequency // todo, since this is a common pattern this enum should probably be shared
+	{
+		// Create a Node Instance of the appropriate child at creation time
+		OnCreate,
+		// Check the input bool every frame, and when its value changes, transition to the appropriate child node
+		OnUpdate
+	};
+	
+	// BlendByBool Anim Node selects between two child nodes based on a bool
+	USTRUCT(DisplayName = "Blend By Bool")
+	struct FUAFBlendByBoolNodeData : public FUAFAnimNodeData
+	{
+		GENERATED_BODY()
+
+		// Bool variable to select based on 
+		// todo: this should be some sort of binding when we build that system
+		UPROPERTY(EditAnywhere, Category = "NodeData")
+		FAnimNextVariableReference Bool;
+
+		// Evaluation Frequency: should we select the appriate child once at creation time, or every frame.
+		UPROPERTY(EditAnywhere, Category = "NodeData")
+		EBlendByBoolEvaluationFrequency EvaluationFrequency = EBlendByBoolEvaluationFrequency::OnUpdate;
+
+		// Child Node to instantiate if input bool is true
+		UPROPERTY(EditAnywhere, Category = "NodeData")
+		TInstancedStruct<FUAFAnimNodeData> TrueNode;
+		
+		// Child Node to instantiate if input bool is false 
+		UPROPERTY(EditAnywhere, Category = "NodeData")
+		TInstancedStruct<FUAFAnimNodeData> FalseNode;
+		
+		// Transition and blend data for transitioning between child nodes
+		UPROPERTY(EditAnywhere, Category = "NodeData")
+		TInstancedStruct<FUAFTransitionNodeData> Transition;
+
+		// FUAFAnimNodeData impl
+		UE_API virtual FUAFAnimNodePtr CreateInstance(FUAFAnimGraphUpdateContext& Context) const override;
+	};
+
+	// BlendByBool Instance
+	struct FUAFBlendByBoolNode : public FUAFBlendStack
+	{
+		UE_API explicit FUAFBlendByBoolNode(FUAFAnimGraphUpdateContext& Context, const FUAFBlendByBoolNodeData* InData);
+
+		// FUAFAnimNode impl
+		UE_API virtual void PreUpdate(FUAFAnimGraphUpdateContext& GraphContext) override;
+
+	private:
+		const FUAFBlendByBoolNodeData* Data = nullptr;
+		bool bCurrentValue = false;
+	};
+}
+
+#undef UE_API
