Move UAF utilities into broader scope

#Synced-CL 51666402

diff --git a/Engine/Plugins/Experimental/UAF/UAFAnimNode/Source/UAFAnimNode/Private/UAF/AnimOpCore/UAFStack.cpp b/Engine/Plugins/Experimental/UAF/UAFAnimNode/Source/UAFAnimNode/Private/UAF/AnimOpCore/UAFStack.cpp
deleted file mode 100644
diff --git a/Engine/Plugins/Experimental/UAF/UAFAnimNode/Source/UAFAnimNode/Private/UAF/UAFStack.cpp b/Engine/Plugins/Experimental/UAF/UAFAnimNode/Source/UAFAnimNode/Private/UAF/UAFStack.cpp
new file mode 100644
--- /dev/null
+++ b/Engine/Plugins/Experimental/UAF/UAFAnimNode/Source/UAFAnimNode/Private/UAF/UAFStack.cpp
@@ -0,0 +1,28 @@
+// Copyright Epic Games, Inc. All Rights Reserved.
+
+#include "UAF/UAFStack.h"
+
+namespace UE::UAF
+{
+	const FUAFStack::FEntryHeader* FUAFStack::PeekImpl(uint32 OffsetFromTop) const
+	{
+		const FEntryHeader* EntryHeader = Top;
+
+		uint32 Count = 0;
+		while (EntryHeader != nullptr && Count < OffsetFromTop)
+		{
+			EntryHeader = EntryHeader->PrevTop;
+			Count++;
+		}
+
+		if (Count != OffsetFromTop)
+		{
+			// Offset is too far, not enough entries
+			checkf(false, TEXT("Attempting to peek at an element that exceeds the stack size."));
+			return nullptr;
+		}
+
+		check(EntryHeader != nullptr);
+		return EntryHeader;
+	}
+}
diff --git a/Engine/Plugins/Experimental/UAF/UAFAnimNode/Source/UAFAnimNode/Public/UAF/AnimOpCore/UAFAnimOpEvaluator.h b/Engine/Plugins/Experimental/UAF/UAFAnimNode/Source/UAFAnimNode/Public/UAF/AnimOpCore/UAFAnimOpEvaluator.h
--- a/Engine/Plugins/Experimental/UAF/UAFAnimNode/Source/UAFAnimNode/Public/UAF/AnimOpCore/UAFAnimOpEvaluator.h
+++ b/Engine/Plugins/Experimental/UAF/UAFAnimNode/Source/UAFAnimNode/Public/UAF/AnimOpCore/UAFAnimOpEvaluator.h
@@ -3,7 +3,7 @@
 #pragma once
 
 #include "CoreMinimal.h"
-#include "UAF/AnimOpCore/UAFStack.h"
+#include "UAF/UAFStack.h"
 
 #define UE_API UAFANIMNODE_API
 
diff --git a/Engine/Plugins/Experimental/UAF/UAFAnimNode/Source/UAFAnimNode/Public/UAF/AnimOpCore/UAFAnimOpValueEvaluator.h b/Engine/Plugins/Experimental/UAF/UAFAnimNode/Source/UAFAnimNode/Public/UAF/AnimOpCore/UAFAnimOpValueEvaluator.h
--- a/Engine/Plugins/Experimental/UAF/UAFAnimNode/Source/UAFAnimNode/Public/UAF/AnimOpCore/UAFAnimOpValueEvaluator.h
+++ b/Engine/Plugins/Experimental/UAF/UAFAnimNode/Source/UAFAnimNode/Public/UAF/AnimOpCore/UAFAnimOpValueEvaluator.h
@@ -4,8 +4,8 @@
 
 #include "CoreMinimal.h"
 #include "UAF/AnimOpCore/UAFAnimOpEvaluator.h"
-#include "UAF/AnimOpCore/UAFCoWRef.h"
 #include "UAF/Attributes/AttributeBindingData.h"
+#include "UAF/UAFCoWRef.h"
 #include "UAF/ValueRuntime/Transformers/Transformer.h"
 #include "UAF/ValueRuntime/ValueBundle.h"
 
diff --git a/Engine/Plugins/Experimental/UAF/UAFAnimNode/Source/UAFAnimNode/Public/UAF/AnimOpCore/UAFCoWRef.h b/Engine/Plugins/Experimental/UAF/UAFAnimNode/Source/UAFAnimNode/Public/UAF/AnimOpCore/UAFCoWRef.h
deleted file mode 100644
diff --git a/Engine/Plugins/Experimental/UAF/UAFAnimNode/Source/UAFAnimNode/Public/UAF/AnimOpCore/UAFStack.h b/Engine/Plugins/Experimental/UAF/UAFAnimNode/Source/UAFAnimNode/Public/UAF/AnimOpCore/UAFStack.h
deleted file mode 100644
diff --git a/Engine/Plugins/Experimental/UAF/UAFAnimNode/Source/UAFAnimNode/Public/UAF/UAFCoWRef.h b/Engine/Plugins/Experimental/UAF/UAFAnimNode/Source/UAFAnimNode/Public/UAF/UAFCoWRef.h
new file mode 100644
--- /dev/null
+++ b/Engine/Plugins/Experimental/UAF/UAFAnimNode/Source/UAFAnimNode/Public/UAF/UAFCoWRef.h
@@ -0,0 +1,394 @@
+// Copyright Epic Games, Inc. All Rights Reserved.
+
+#pragma once
+
+#include "CoreMinimal.h"
+#include "UAF/ValueRuntime/IndirectAllocator.h"
+
+namespace UE::UAF
+{
+	/**
+	 * TUAFCoWRef
+	 *
+	 * A reference to a value type that performs Copy-on-Write.
+	 * A reference is never invalid, it must always be constructed from valid data.
+	 * 
+	 * A CoW Reference tracks 3 things:
+	 *   - A pointer to an allocated value
+	 *   - Whether or not it owns the allocated memory (if we own it, the reference is unique)
+	 *   - Whether or not the reference is immutable
+	 * 
+	 * For performance reasons, this tracking is performed using a single pointer where the bottom two bits
+	 * are used to store our metadata. We thus require an alignment of 4 bytes for all allocations.
+	 * 
+	 * Mutability is enforced at runtime not through the type system (e.g. with 'const').
+	 * The specified allocator will be used when a copy is made and memory is owned by the reference.
+	 * 
+	 * Copying or assigning a reference to another always performs a copy where the new reference
+	 * owns the resulting memory. This occurs regardless whether or not the source/target are mutable/immutable.
+	 * 
+	 * It is always safe to query for a const reference using TUAFCoWRef::Get() regardless of whether or
+	 * not the underlying reference is immutable. However, care must be taken when using TUAFCoWRef::GetMutable().
+	 * It is the responsibility of the caller to ensure that the reference is mutable before attempting to
+	 * use it as such. If the reference is not mutable, TUAFCoWRef::ForceMutable() can be used to convert it.
+	 * 
+	 * Two allocators are supported: FDefaultAllocator and TMemStackAllocator.
+	 */
+	template<class ValueType, class AllocatorType = FDefaultAllocator>
+	class TUAFCoWRef
+	{
+	public:
+		// Constructs a new mutable reference from the specified value by performing a move of the input.
+		// The reference will be unique as its memory will be owned by us.
+		[[nodiscard]] static TUAFCoWRef<ValueType, AllocatorType> MakeFrom(ValueType&& Value);
+
+		// Constructs a new mutable reference from the specified value pointer.
+		// The reference will not be unique as its memory will not be owned by us.
+		[[nodiscard]] static TUAFCoWRef<ValueType, AllocatorType> MakeMutable(ValueType* ValuePtr);
+
+		// Constructs a new immutable reference from the specified value pointer.
+		// The reference will not be unique as its memory will not be owned by us.
+		[[nodiscard]] static TUAFCoWRef<ValueType, AllocatorType> MakeImmutable(const ValueType* ValuePtr);
+
+		// Constructs a new mutable reference by performing a copy of the input reference.
+		// The reference will be unique as its memory will be owned by us.
+		TUAFCoWRef(const TUAFCoWRef<ValueType, AllocatorType>& Other);
+
+		// Constructs a new mutable reference by performing a copy of the input reference.
+		// The reference will be unique as its memory will be owned by us.
+		template<class OtherAllocatorType>
+		TUAFCoWRef(const TUAFCoWRef<ValueType, OtherAllocatorType>& Other);
+
+		// Move a reference into another
+		TUAFCoWRef(TUAFCoWRef<ValueType, AllocatorType>&& Other);
+
+		// Destroys the current reference
+		// If it owned the memory within, it is freed
+		~TUAFCoWRef();
+		
+		// Assign a new mutable reference by performing a copy of the input reference.
+		// The reference will be unique as its memory will be owned by us.
+		TUAFCoWRef<ValueType, AllocatorType>& operator=(const TUAFCoWRef<ValueType, AllocatorType>& Other);
+
+		// Assign a new mutable reference by performing a copy of the input reference.
+		// The reference will be unique as its memory will be owned by us.
+		template<class OtherAllocatorType>
+		TUAFCoWRef<ValueType, AllocatorType>& operator=(const TUAFCoWRef<ValueType, OtherAllocatorType>& Other);
+
+		// Move assign a reference into another
+		TUAFCoWRef<ValueType, AllocatorType>& operator=(TUAFCoWRef<ValueType, AllocatorType>&& Other);
+
+		// Move assign a value into a reference by performing a move of the input.
+		// The reference will be unique as its memory will be owned by us.
+		TUAFCoWRef<ValueType, AllocatorType>& operator=(ValueType&& Value);
+
+		// Returns an immutable reference to the value contained within
+		[[nodiscard]] const ValueType& Get() const;
+
+		// Returns a mutable reference to the value contained within
+		[[nodiscard]] ValueType& GetMutable();
+
+		// Ensures that this reference is mutable
+		// If the current reference was immutable and we do not own the memory, this performs a copy (Copy-on-Write)
+		// Note that any existing const references will remain valid since the original owner is keeping it alive
+		void ForceMutable();
+
+		// Returns whether or not this reference is mutable
+		[[nodiscard]] bool IsMutable() const;
+
+		// Returns whether or not this reference is unique
+		// A reference is unique if we own its memory
+		[[nodiscard]] bool IsUnique() const;
+
+	private:
+		static constexpr uintptr_t VALUE_PTR_MASK = static_cast<uintptr_t>(~0x0llu & ~0x03llu);
+		static constexpr uintptr_t FLAGS_MASK = 0x03llu;
+		static constexpr uintptr_t IS_MUTABLE_FLAG_BIT = 0x02llu;
+		static constexpr uintptr_t OWNS_MEMORY_FLAG_BIT = 0x01llu;
+
+		// Disable construction of empty references
+		TUAFCoWRef() = default;
+
+		ValueType* GetValuePtr();
+		const ValueType* GetValuePtr() const;
+
+		// Stores a pointer to a value and 2 flags in the LSBs
+		uintptr_t ValuePtrAndFlags = 0;
+	};
+
+	// Returns a mutable reference from a list of potential candidates.
+	// If candidates are mutable, the first one will be returned.
+	// Otherwise, the first unique immutable reference is returned.
+	// If none of the candidates are mutable or unique, we perform a Copy-on-Write with the last immutable reference.
+	template<class ValueType, class AllocatorType>
+	TUAFCoWRef<ValueType, AllocatorType>& FindMutableCoWRef(TUAFCoWRef<ValueType, AllocatorType>& Arg0, TUAFCoWRef<ValueType, AllocatorType>& Arg1);
+
+	//////////////////////////////////////////////////////////////////////////
+	// Inline implementation
+
+	template<class ValueType, class AllocatorType>
+	inline TUAFCoWRef<ValueType, AllocatorType> TUAFCoWRef<ValueType, AllocatorType>::MakeFrom(ValueType&& Value)
+	{
+		TUAFCoWRef<ValueType, AllocatorType> Result;
+
+		uint8* CopyPtr = FAllocatorTypeTrait<AllocatorType>::Realloc(nullptr, 0, sizeof(ValueType));
+		new(CopyPtr) ValueType(MoveTemp(Value));
+
+		Result.ValuePtrAndFlags = reinterpret_cast<uintptr_t>(CopyPtr);
+		checkf((Result.ValuePtrAndFlags & FLAGS_MASK) == 0, TEXT("Pointer must be aligned to 4 bytes to support CoW."));
+		Result.ValuePtrAndFlags |= IS_MUTABLE_FLAG_BIT;
+		Result.ValuePtrAndFlags |= OWNS_MEMORY_FLAG_BIT;
+
+		return Result;
+	}
+
+	template<class ValueType, class AllocatorType>
+	inline TUAFCoWRef<ValueType, AllocatorType> TUAFCoWRef<ValueType, AllocatorType>::MakeMutable(ValueType* ValuePtr)
+	{
+		checkf(ValuePtr != nullptr, TEXT("Cannot create a CoW reference from a null pointer."));
+
+		TUAFCoWRef<ValueType, AllocatorType> Result;
+
+		Result.ValuePtrAndFlags = reinterpret_cast<uintptr_t>(ValuePtr);
+		checkf((Result.ValuePtrAndFlags & FLAGS_MASK) == 0, TEXT("Pointer must be aligned to 4 bytes to support CoW."));
+		Result.ValuePtrAndFlags |= IS_MUTABLE_FLAG_BIT;
+
+		return Result;
+	}
+
+	template<class ValueType, class AllocatorType>
+	inline TUAFCoWRef<ValueType, AllocatorType> TUAFCoWRef<ValueType, AllocatorType>::MakeImmutable(const ValueType* ValuePtr)
+	{
+		checkf(ValuePtr != nullptr, TEXT("Cannot create a CoW reference from a null pointer."));
+
+		TUAFCoWRef<ValueType, AllocatorType> Result;
+
+		Result.ValuePtrAndFlags = reinterpret_cast<uintptr_t>(ValuePtr);
+		checkf((Result.ValuePtrAndFlags & FLAGS_MASK) == 0, TEXT("Pointer must be aligned to 4 bytes to support CoW."));
+
+		return Result;
+	}
+
+	template<class ValueType, class AllocatorType>
+	inline TUAFCoWRef<ValueType, AllocatorType>::TUAFCoWRef(const TUAFCoWRef<ValueType, AllocatorType>& Other)
+	{
+		uint8* CopyPtr = FAllocatorTypeTrait<AllocatorType>::Realloc(nullptr, 0, sizeof(ValueType));
+		new(CopyPtr) ValueType(*Other.GetValuePtr());
+
+		ValuePtrAndFlags = reinterpret_cast<uintptr_t>(CopyPtr);
+		checkf((ValuePtrAndFlags & FLAGS_MASK) == 0, TEXT("Pointer must be aligned to 4 bytes to support CoW."));
+		ValuePtrAndFlags |= IS_MUTABLE_FLAG_BIT;
+		ValuePtrAndFlags |= OWNS_MEMORY_FLAG_BIT;
+	}
+
+	template<class ValueType, class AllocatorType>
+	template<class OtherAllocatorType>
+	inline TUAFCoWRef<ValueType, AllocatorType>::TUAFCoWRef(const TUAFCoWRef<ValueType, OtherAllocatorType>& Other)
+	{
+		uint8* CopyPtr = FAllocatorTypeTrait<AllocatorType>::Realloc(nullptr, 0, sizeof(ValueType));
+		new(CopyPtr) ValueType(*Other.GetValuePtr());
+
+		ValuePtrAndFlags = reinterpret_cast<uintptr_t>(CopyPtr);
+		checkf((ValuePtrAndFlags & FLAGS_MASK) == 0, TEXT("Pointer must be aligned to 4 bytes to support CoW."));
+		ValuePtrAndFlags |= IS_MUTABLE_FLAG_BIT;
+		ValuePtrAndFlags |= OWNS_MEMORY_FLAG_BIT;
+	}
+
+	template<class ValueType, class AllocatorType>
+	inline TUAFCoWRef<ValueType, AllocatorType>::TUAFCoWRef(TUAFCoWRef<ValueType, AllocatorType>&& Other)
+		: ValuePtrAndFlags(Other.ValuePtrAndFlags)
+	{
+		Other.ValuePtrAndFlags = 0;
+	}
+
+	template<class ValueType, class AllocatorType>
+	inline TUAFCoWRef<ValueType, AllocatorType>::~TUAFCoWRef()
+	{
+		if ((ValuePtrAndFlags & OWNS_MEMORY_FLAG_BIT) != 0)
+		{
+			ValueType* ValuePtr = GetValuePtr();
+
+			ValuePtr->~ValueType();
+			FAllocatorTypeTrait<AllocatorType>::Realloc(reinterpret_cast<uint8*>(ValuePtr), sizeof(ValueType), 0);
+		}
+	}
+
+	template<class ValueType, class AllocatorType>
+	inline TUAFCoWRef<ValueType, AllocatorType>& TUAFCoWRef<ValueType, AllocatorType>::operator=(const TUAFCoWRef<ValueType, AllocatorType>& Other)
+	{
+		if (this != &Other)
+		{
+			if ((ValuePtrAndFlags & OWNS_MEMORY_FLAG_BIT) != 0)
+			{
+				ValueType* ValuePtr = GetValuePtr();
+
+				ValuePtr->~ValueType();
+				FAllocatorTypeTrait<AllocatorType>::Realloc(reinterpret_cast<uint8*>(ValuePtr), sizeof(ValueType), 0);
+			}
+
+			uint8* CopyPtr = FAllocatorTypeTrait<AllocatorType>::Realloc(nullptr, 0, sizeof(ValueType));
+			new(CopyPtr) ValueType(*Other.GetValuePtr());
+
+			ValuePtrAndFlags = reinterpret_cast<uintptr_t>(CopyPtr);
+			checkf((ValuePtrAndFlags & FLAGS_MASK) == 0, TEXT("Pointer must be aligned to 4 bytes to support CoW."));
+			ValuePtrAndFlags |= IS_MUTABLE_FLAG_BIT;
+			ValuePtrAndFlags |= OWNS_MEMORY_FLAG_BIT;
+		}
+
+		return *this;
+	}
+
+	template<class ValueType, class AllocatorType>
+	template<class OtherAllocatorType>
+	inline TUAFCoWRef<ValueType, AllocatorType>& TUAFCoWRef<ValueType, AllocatorType>::operator=(const TUAFCoWRef<ValueType, OtherAllocatorType>& Other)
+	{
+		if ((ValuePtrAndFlags & OWNS_MEMORY_FLAG_BIT) != 0)
+		{
+			ValueType* ValuePtr = GetValuePtr();
+
+			ValuePtr->~ValueType();
+			FAllocatorTypeTrait<AllocatorType>::Realloc(reinterpret_cast<uint8*>(ValuePtr), sizeof(ValueType), 0);
+		}
+
+		uint8* CopyPtr = FAllocatorTypeTrait<AllocatorType>::Realloc(nullptr, 0, sizeof(ValueType));
+		new(CopyPtr) ValueType(*Other.GetValuePtr());
+
+		ValuePtrAndFlags = reinterpret_cast<uintptr_t>(CopyPtr);
+		checkf((ValuePtrAndFlags & FLAGS_MASK) == 0, TEXT("Pointer must be aligned to 4 bytes to support CoW."));
+		ValuePtrAndFlags |= IS_MUTABLE_FLAG_BIT;
+		ValuePtrAndFlags |= OWNS_MEMORY_FLAG_BIT;
+
+		return *this;
+	}
+
+	template<class ValueType, class AllocatorType>
+	inline TUAFCoWRef<ValueType, AllocatorType>& TUAFCoWRef<ValueType, AllocatorType>::operator=(TUAFCoWRef<ValueType, AllocatorType>&& Other)
+	{
+		if (this != &Other)
+		{
+			if ((ValuePtrAndFlags & OWNS_MEMORY_FLAG_BIT) != 0)
+			{
+				ValueType* ValuePtr = GetValuePtr();
+
+				ValuePtr->~ValueType();
+				FAllocatorTypeTrait<AllocatorType>::Realloc(reinterpret_cast<uint8*>(ValuePtr), sizeof(ValueType), 0);
+			}
+
+			ValuePtrAndFlags = Other.ValuePtrAndFlags;
+			Other.ValuePtrAndFlags = 0;
+		}
+
+		return *this;
+	}
+
+	template<class ValueType, class AllocatorType>
+	inline TUAFCoWRef<ValueType, AllocatorType>& TUAFCoWRef<ValueType, AllocatorType>::operator=(ValueType&& Value)
+	{
+		if ((ValuePtrAndFlags & OWNS_MEMORY_FLAG_BIT) != 0)
+		{
+			ValueType* ValuePtr = GetValuePtr();
+
+			ValuePtr->~ValueType();
+			FAllocatorTypeTrait<AllocatorType>::Realloc(reinterpret_cast<uint8*>(ValuePtr), sizeof(ValueType), 0);
+		}
+
+		uint8* CopyPtr = FAllocatorTypeTrait<AllocatorType>::Realloc(nullptr, 0, sizeof(ValueType));
+		new(CopyPtr) ValueType(MoveTemp(Value));
+
+		ValuePtrAndFlags = reinterpret_cast<uintptr_t>(CopyPtr);
+		checkf((ValuePtrAndFlags & FLAGS_MASK) == 0, TEXT("Pointer must be aligned to 4 bytes to support CoW."));
+		ValuePtrAndFlags |= IS_MUTABLE_FLAG_BIT;
+		ValuePtrAndFlags |= OWNS_MEMORY_FLAG_BIT;
+
+		return *this;
+	}
+
+	template<class ValueType, class AllocatorType>
+	inline const ValueType& TUAFCoWRef<ValueType, AllocatorType>::Get() const
+	{
+		return *GetValuePtr();
+	}
+
+	template<class ValueType, class AllocatorType>
+	inline ValueType& TUAFCoWRef<ValueType, AllocatorType>::GetMutable()
+	{
+		checkf((ValuePtrAndFlags & IS_MUTABLE_FLAG_BIT) != 0, TEXT("Attempting to get a mutable reference from an immutable CoW Ref."));
+		return *GetValuePtr();
+	}
+
+	template<class ValueType, class AllocatorType>
+	inline void TUAFCoWRef<ValueType, AllocatorType>::ForceMutable()
+	{
+		if ((ValuePtrAndFlags & IS_MUTABLE_FLAG_BIT) != 0) [[likely]]
+		{
+			// We are mutable, all good
+			return;
+		}
+		else
+		{
+			// We are immutable
+			if ((ValuePtrAndFlags & OWNS_MEMORY_FLAG_BIT) != 0)
+			{
+				// We own the memory, mark is as mutable
+				ValuePtrAndFlags |= IS_MUTABLE_FLAG_BIT;
+				return;
+			}
+
+			// We do not own the memory, make a copy (Copy-on-Write)
+			uint8* CopyPtr = FAllocatorTypeTrait<AllocatorType>::Realloc(nullptr, 0, sizeof(ValueType));
+			new(CopyPtr) ValueType(*GetValuePtr());
+
+			ValuePtrAndFlags = reinterpret_cast<uintptr_t>(CopyPtr);
+			checkf((ValuePtrAndFlags & FLAGS_MASK) == 0, TEXT("Pointer must be aligned to 4 bytes to support CoW."));
+			ValuePtrAndFlags |= IS_MUTABLE_FLAG_BIT;
+			ValuePtrAndFlags |= OWNS_MEMORY_FLAG_BIT;
+		}
+	}
+
+	template<class ValueType, class AllocatorType>
+	inline bool TUAFCoWRef<ValueType, AllocatorType>::IsMutable() const
+	{
+		return (ValuePtrAndFlags & IS_MUTABLE_FLAG_BIT) != 0;
+	}
+
+	template<class ValueType, class AllocatorType>
+	inline bool TUAFCoWRef<ValueType, AllocatorType>::IsUnique() const
+	{
+		return (ValuePtrAndFlags & OWNS_MEMORY_FLAG_BIT) != 0;
+	}
+
+	template<class ValueType, class AllocatorType>
+	inline ValueType* TUAFCoWRef<ValueType, AllocatorType>::GetValuePtr()
+	{
+		return reinterpret_cast<ValueType*>(ValuePtrAndFlags & VALUE_PTR_MASK);
+	}
+
+	template<class ValueType, class AllocatorType>
+	inline const ValueType* TUAFCoWRef<ValueType, AllocatorType>::GetValuePtr() const
+	{
+		return reinterpret_cast<const ValueType*>(ValuePtrAndFlags & VALUE_PTR_MASK);
+	}
+
+	template<class ValueType, class AllocatorType>
+	inline TUAFCoWRef<ValueType, AllocatorType>& FindMutableCoWRef(TUAFCoWRef<ValueType, AllocatorType>& Arg0, TUAFCoWRef<ValueType, AllocatorType>& Arg1)
+	{
+		if (Arg0.IsMutable())
+		{
+			return Arg0;
+		}
+
+		if (Arg1.IsMutable())
+		{
+			return Arg1;
+		}
+
+		if (Arg0.IsUnique())
+		{
+			Arg0.ForceMutable();
+			return Arg0;
+		}
+
+		Arg1.ForceMutable();
+		return Arg1;
+	}
+}
diff --git a/Engine/Plugins/Experimental/UAF/UAFAnimNode/Source/UAFAnimNode/Public/UAF/UAFStack.h b/Engine/Plugins/Experimental/UAF/UAFAnimNode/Source/UAFAnimNode/Public/UAF/UAFStack.h
new file mode 100644
--- /dev/null
+++ b/Engine/Plugins/Experimental/UAF/UAFAnimNode/Source/UAFAnimNode/Public/UAF/UAFStack.h
@@ -0,0 +1,277 @@
+// Copyright Epic Games, Inc. All Rights Reserved.
+
+#pragma once
+
+#include "CoreMinimal.h"
+#include "ConstExprUID.h"
+
+#define UE_API UAFANIMNODE_API
+
+// If you wish to use a custom type with a stack when evaluating AnimOps, you'll have to enable
+// its usage through this macro to ensure that GetStackTypeID<T>() is declared for your type.
+// Use this macro inside the UE::UAF namespace
+#define UE_UAF_REGISTER_STACK_TYPE(Type) \
+	template<> \
+	[[nodiscard]] constexpr uint32 GetStackTypeID<Type>() \
+	{ \
+		return UE::HashStringFNV1a32(#Type); \
+	}
+
+namespace UE::UAF
+{
+	// Helper function that returns a UID for the specified stack element type
+	// @see UE_UAF_REGISTER_STACK_TYPE
+	template<typename Type>
+	[[nodiscard]] constexpr uint32 GetStackTypeID()
+	{
+		static_assert(false, "Not implemented! See UE_UAF_REGISTER_STACK_TYPE for details and register your type");
+		return 0;
+	}
+
+	/*
+	 * FUAFStackName
+	 *
+	 * A small struct to wrap an FName and cache its type hash.
+	 */
+	class FUAFStackName final
+	{
+	public:
+		// Implicitly coerce or construct a stack name from an FName
+		FUAFStackName(const FName& InName)
+			: Name(InName)
+			, NameHash(GetTypeHash(InName))
+		{
+		}
+
+		// Returns the stack name
+		[[nodiscard]] const FName& GetName() const;
+
+		// Returns the stack name hash
+		[[nodiscard]] uint32 GetNameHash() const;
+
+	private:
+		FName Name;
+		uint32 NameHash = 0;
+	};
+
+	/*
+	 * FUAFStack
+	 *
+	 * A simple stack implementation that supports an indirect allocator for its entries.
+	 * Type safety is enforced at runtime and compile time where possible.
+	 * 
+	 * In order to use a stack with a new type, make sure that the type has been registered with
+	 * the following macro: UE_UAF_REGISTER_STACK_TYPE.
+	 */
+	class FUAFStack
+	{
+	public:
+		// Virtual destructor to simplify destruction
+		virtual ~FUAFStack() = default;
+
+		// Returns the stack name
+		[[nodiscard]] const FName& GetName() const;
+
+		// Returns the number of stack entries
+		[[nodiscard]] uint32 Num() const;
+
+		// Returns whether or not the stack is empty
+		[[nodiscard]] bool IsEmpty() const;
+
+		// Returns the stack type ID
+		[[nodiscard]] uint32 GetTypeID() const;
+
+	protected:
+		struct FEntryHeader
+		{
+			// The top entry at the time this new entry was pushed
+			FEntryHeader* PrevTop = nullptr;
+
+			// The allocated size of this entry, includes sizeof(FEntry) and any required alignment
+			uint32 Size = 0;
+
+			// The entry start offset from 'this'
+			uint32 Offset = 0;
+
+			// Actual entry memory follows
+
+			[[nodiscard]] void* GetValuePtr();
+			[[nodiscard]] const void* GetValuePtr() const;
+		};
+
+		FUAFStack(const FUAFStackName& Name, uint32 TypeID);
+		[[nodiscard]] UE_API const FEntryHeader* PeekImpl(uint32 OffsetFromTop) const;
+
+		// The top entry on the stack
+		FEntryHeader* Top = nullptr;
+
+		// The name of this stack
+		FUAFStackName Name;
+
+		// How many entries we have on the stack
+		uint32 NumEntries = 0;
+
+		// Type ID used for runtime checks to make sure we always read/write using the same type
+		uint32 TypeID = 0;
+	};
+
+	/*
+	 * TUAFStack
+	 *
+	 * A simple stack implementation that supports an indirect allocator for its entries.
+	 * Type safety is enforced at runtime and compile time where possible.
+	 *
+	 * In order to use a stack with a new type, make sure that the type has been registered with
+	 * the following macro: UE_UAF_REGISTER_STACK_TYPE.
+	 */
+	template<typename ValueType>
+	class TUAFStack : public FUAFStack
+	{
+	public:
+		// Constructs a named stack
+		TUAFStack(const FUAFStackName& Name);
+		virtual ~TUAFStack() override;
+
+		// Pushes a new value onto the stack
+		void Push(ValueType&& Value);
+
+		// Pops and returns a value from the stack
+		[[nodiscard]] ValueType Pop();
+
+		// Peeks at an element with the specified offset from the top and returns it
+		[[nodiscard]] const ValueType* Peek(uint32 OffsetFromTop = 0) const;
+		[[nodiscard]] ValueType* PeekMutable(uint32 OffsetFromTop = 0);
+	};
+
+	//////////////////////////////////////////////////////////////////////////
+	// Inline implementation
+
+	inline const FName& FUAFStackName::GetName() const
+	{
+		return Name;
+	}
+
+	inline uint32 FUAFStackName::GetNameHash() const
+	{
+		return NameHash;
+	}
+
+	inline const FName& FUAFStack::GetName() const
+	{
+		return Name.GetName();
+	}
+
+	inline uint32 FUAFStack::Num() const
+	{
+		return NumEntries;
+	}
+
+	inline bool FUAFStack::IsEmpty() const
+	{
+		return Top == nullptr;
+	}
+
+	inline uint32 FUAFStack::GetTypeID() const
+	{
+		return TypeID;
+	}
+
+	inline FUAFStack::FUAFStack(const FUAFStackName& InName, uint32 InTypeID)
+		: Name(InName)
+		, TypeID(InTypeID)
+	{
+	}
+
+	inline void* FUAFStack::FEntryHeader::GetValuePtr()
+	{
+		return reinterpret_cast<uint8*>(this) + Offset;
+	}
+
+	inline const void* FUAFStack::FEntryHeader::GetValuePtr() const
+	{
+		return reinterpret_cast<const uint8*>(this) + Offset;
+	}
+
+	template<typename ValueType>
+	inline TUAFStack<ValueType>::TUAFStack(const FUAFStackName& Name)
+		: FUAFStack(Name, GetStackTypeID<ValueType>())
+	{
+	}
+
+	template<typename ValueType>
+	inline TUAFStack<ValueType>::~TUAFStack()
+	{
+		FEntryHeader* EntryHeader = Top;
+		while (EntryHeader != nullptr)
+		{
+			FEntryHeader* PrevEntryHeader = EntryHeader->PrevTop;
+
+			static_cast<ValueType*>(EntryHeader->GetValuePtr())->~ValueType();
+			FMemory::Free(EntryHeader);
+
+			EntryHeader = PrevEntryHeader;
+		}
+	}
+
+	template<typename ValueType>
+	inline void TUAFStack<ValueType>::Push(ValueType&& Value)
+	{
+		checkf(TypeID == GetStackTypeID<ValueType>(), TEXT("Type mismatch! This stack is being queried with a different type than it was created with."));
+
+		constexpr uint32 MinAlign = alignof(ValueType) < alignof(FEntryHeader) ? alignof(FEntryHeader) : alignof(ValueType);
+		const uint32 EntryOffset = Align(sizeof(FEntryHeader), alignof(ValueType));
+		const uint32 AllocationSize = EntryOffset + sizeof(ValueType);
+
+		FEntryHeader* EntryHeader = static_cast<FEntryHeader*>(FMemory::Malloc(AllocationSize, MinAlign));
+		EntryHeader->PrevTop = Top;
+		EntryHeader->Size = AllocationSize;
+		EntryHeader->Offset = EntryOffset;
+
+		// We are pushing an immutable version of our value, move its contents to the stack, we'll own it
+		ValueType* ValuePtr = static_cast<ValueType*>(EntryHeader->GetValuePtr());
+		new(ValuePtr) ValueType(MoveTemp(Value));
+
+		Top = EntryHeader;
+		NumEntries++;
+	}
+
+	template<typename ValueType>
+	inline ValueType TUAFStack<ValueType>::Pop()
+	{
+		checkf(TypeID == GetStackTypeID<ValueType>(), TEXT("Type mismatch! This stack is being queried with a different type than it was created with."));
+		checkf(Top != nullptr, TEXT("Attempting to pop an element from an empty stack."));
+
+		FEntryHeader* EntryHeader = Top;
+		Top = EntryHeader->PrevTop;
+		NumEntries--;
+
+		ValueType* ValuePtr = static_cast<ValueType*>(EntryHeader->GetValuePtr());
+
+		ValueType Result = MoveTemp(*ValuePtr);
+
+		ValuePtr->~ValueType();
+		FMemory::Free(EntryHeader);
+
+		return Result;
+	}
+
+	template<typename ValueType>
+	inline const ValueType* TUAFStack<ValueType>::Peek(uint32 OffsetFromTop) const
+	{
+		checkf(TypeID == GetStackTypeID<ValueType>(), TEXT("Type mismatch! This stack is being queried with a different type than it was created with."));
+		checkf(Top != nullptr, TEXT("Attempting to peek at an element from an empty stack."));
+
+		return static_cast<const ValueType*>(PeekImpl(OffsetFromTop)->GetValuePtr());
+	}
+
+	template<typename ValueType>
+	inline ValueType* TUAFStack<ValueType>::PeekMutable(uint32 OffsetFromTop)
+	{
+		checkf(TypeID == GetStackTypeID<ValueType>(), TEXT("Type mismatch! This stack is being queried with a different type than it was created with."));
+		checkf(Top != nullptr, TEXT("Attempting to peek at an element from an empty stack."));
+
+		return static_cast<ValueType*>(const_cast<FEntryHeader*>(PeekImpl(OffsetFromTop))->GetValuePtr());
+	}
+}
+
+#undef UE_API
diff --git a/Engine/Plugins/Experimental/UAF/UAFPoseSearch/Source/UAFPoseSearch/Private/AnimNode/UAFPoseHistoryNode.cpp b/Engine/Plugins/Experimental/UAF/UAFPoseSearch/Source/UAFPoseSearch/Private/AnimNode/UAFPoseHistoryNode.cpp
--- a/Engine/Plugins/Experimental/UAF/UAFPoseSearch/Source/UAFPoseSearch/Private/AnimNode/UAFPoseHistoryNode.cpp
+++ b/Engine/Plugins/Experimental/UAF/UAFPoseSearch/Source/UAFPoseSearch/Private/AnimNode/UAFPoseHistoryNode.cpp
@@ -6,7 +6,6 @@
 #include "UAFAssetInstance.h"
 #include "UAF/AnimOpCore/UAFAnimOpValueEvaluator.h"
 #include "UAF/AnimOpCore/UAFAnimOp.h"
-#include "UAF/AnimOpCore/UAFStack.h"
 #include "UAF/AnimOps/UAFNullAnimOp.h"
 #include "UAF/Attributes/EngineAttributes.h"
 
