Properly forward arguments when creating anim nodes in new UAF runtime

#Synced-CL 51666402

diff --git a/Engine/Plugins/Experimental/UAF/UAFAnimNode/Source/UAFAnimNode/Public/UAF/AnimNodeCore/UAFAnimNode.h b/Engine/Plugins/Experimental/UAF/UAFAnimNode/Source/UAFAnimNode/Public/UAF/AnimNodeCore/UAFAnimNode.h
--- a/Engine/Plugins/Experimental/UAF/UAFAnimNode/Source/UAFAnimNode/Public/UAF/AnimNodeCore/UAFAnimNode.h
+++ b/Engine/Plugins/Experimental/UAF/UAFAnimNode/Source/UAFAnimNode/Public/UAF/AnimNodeCore/UAFAnimNode.h
@@ -212,7 +212,7 @@ namespace UE::UAF
 
 	// Constructs a new anim node of the specified type with the provided arguments and returns a smart pointer to it
 	template<class AnimNodeType, typename... ArgTypes>
-	FUAFAnimNodePtr MakeAnimNode(FUAFAnimGraphUpdateContext& UpdateContext, ArgTypes... Args);
+	FUAFAnimNodePtr MakeAnimNode(FUAFAnimGraphUpdateContext& UpdateContext, ArgTypes&&... Args);
 
 	//////////////////////////////////////////////////////////////////////////
 	// Inline implementation
@@ -389,7 +389,7 @@ namespace UE::UAF
 	}
 
 	template<class AnimNodeType, typename... ArgTypes>
-	inline FUAFAnimNodePtr MakeAnimNode(FUAFAnimGraphUpdateContext& UpdateContext, ArgTypes... Args)
+	inline FUAFAnimNodePtr MakeAnimNode(FUAFAnimGraphUpdateContext& UpdateContext, ArgTypes&&... Args)
 	{
 		// No need to increment the reference, it initializes to 1 already for us
 		constexpr bool bAddRef = false;
