Prevent private symbol leakage

#Synced-CL 51666402

diff --git a/Engine/Plugins/Experimental/UAF/UAFAnimGraph/Source/UAFAnimGraphEditor/Private/RewindDebugger/SequenceInfoTrack.cpp b/Engine/Plugins/Experimental/UAF/UAFAnimGraph/Source/UAFAnimGraphEditor/Private/RewindDebugger/SequenceInfoTrack.cpp
--- a/Engine/Plugins/Experimental/UAF/UAFAnimGraph/Source/UAFAnimGraphEditor/Private/RewindDebugger/SequenceInfoTrack.cpp
+++ b/Engine/Plugins/Experimental/UAF/UAFAnimGraph/Source/UAFAnimGraphEditor/Private/RewindDebugger/SequenceInfoTrack.cpp
@@ -53,7 +53,7 @@ void FAnimNextSequenceTraceInfoCustomization::CustomizeHeader(TSharedRef<IProper
 
 // Copied from SAnimNotifyPanel
 // Todo: find common place for this to live
-FLinearColor GenerateColorFromName(FName Name)
+static FLinearColor GenerateColorFromName(FName Name)
 {
 	constexpr uint8 Saturation = 255;
 	constexpr uint8 Luminosity = 255;
diff --git a/Engine/Plugins/Experimental/UAF/UAFAnimNode/Source/UAFAnimNodeEditor/Private/RewindDebugger/UAFSequenceInfoTrack.cpp b/Engine/Plugins/Experimental/UAF/UAFAnimNode/Source/UAFAnimNodeEditor/Private/RewindDebugger/UAFSequenceInfoTrack.cpp
--- a/Engine/Plugins/Experimental/UAF/UAFAnimNode/Source/UAFAnimNodeEditor/Private/RewindDebugger/UAFSequenceInfoTrack.cpp
+++ b/Engine/Plugins/Experimental/UAF/UAFAnimNode/Source/UAFAnimNodeEditor/Private/RewindDebugger/UAFSequenceInfoTrack.cpp
@@ -53,7 +53,7 @@ void FUAFSequenceTraceInfoCustomization::CustomizeHeader(TSharedRef<IPropertyHan
 
 // Copied from SAnimNotifyPanel
 // Todo: find common place for this to live
-FLinearColor GenerateColorFromName(FName Name)
+static FLinearColor GenerateColorFromName(FName Name)
 {
 	constexpr uint8 Saturation = 255;
 	constexpr uint8 Luminosity = 255;
