Fix UAFAnimNodeEditor use-after-free during shutdown

#rnx
#Synced-CL 52060484

diff --git a/Engine/Plugins/Experimental/UAF/UAFAnimNode/Source/UAFAnimNodeEditor/Private/UAFAnimNodeEditorModule.cpp b/Engine/Plugins/Experimental/UAF/UAFAnimNode/Source/UAFAnimNodeEditor/Private/UAFAnimNodeEditorModule.cpp
--- a/Engine/Plugins/Experimental/UAF/UAFAnimNode/Source/UAFAnimNodeEditor/Private/UAFAnimNodeEditorModule.cpp
+++ b/Engine/Plugins/Experimental/UAF/UAFAnimNode/Source/UAFAnimNodeEditor/Private/UAFAnimNodeEditorModule.cpp
@@ -41,8 +41,10 @@ void FModule::StartupModule()
 
 	FPropertyEditorModule& PropertyModule = FModuleManager::LoadModuleChecked<FPropertyEditorModule>("PropertyEditor");
 
+	CachedSequenceTraceInfoName = FUAFSequenceTraceInfo::StaticStruct()->GetFName();
+
 	PropertyModule.RegisterCustomPropertyTypeLayout(
-		FUAFSequenceTraceInfo::StaticStruct()->GetFName(),
+		CachedSequenceTraceInfoName,
 		FOnGetPropertyTypeCustomizationInstance::CreateLambda([] { return MakeShared<UE::UAF::Editor::FUAFSequenceTraceInfoCustomization>(); }));
 
 	IModularFeatures::Get().RegisterModularFeature(TraceServices::ModuleFeatureName, &GUAFAnimNodeTraceModule);
@@ -63,7 +65,7 @@ void FModule::ShutdownModule()
 	if (FModuleManager::Get().IsModuleLoaded("PropertyEditor"))
 	{
 		FPropertyEditorModule& PropertyModule = FModuleManager::LoadModuleChecked<FPropertyEditorModule>("PropertyEditor");
-		PropertyModule.UnregisterCustomPropertyTypeLayout(FUAFSequenceTraceInfo::StaticStruct()->GetFName());
+		PropertyModule.UnregisterCustomPropertyTypeLayout(CachedSequenceTraceInfoName);
 	}
 
 	FUAFAnimNodeDataFactory::UnregisterAsset(AnimSequenceClassPath);
diff --git a/Engine/Plugins/Experimental/UAF/UAFAnimNode/Source/UAFAnimNodeEditor/Private/UAFAnimNodeEditorModule.h b/Engine/Plugins/Experimental/UAF/UAFAnimNode/Source/UAFAnimNodeEditor/Private/UAFAnimNodeEditorModule.h
--- a/Engine/Plugins/Experimental/UAF/UAFAnimNode/Source/UAFAnimNodeEditor/Private/UAFAnimNodeEditorModule.h
+++ b/Engine/Plugins/Experimental/UAF/UAFAnimNode/Source/UAFAnimNodeEditor/Private/UAFAnimNodeEditorModule.h
@@ -16,6 +16,7 @@ private:
 	virtual void ShutdownModule() override;
 
 	FTopLevelAssetPath AnimSequenceClassPath;
+	FName CachedSequenceTraceInfoName;
 };
 
 }
