[Main] -  Removed API export from inline functions.  Expose API though minimalapi

#Synced-CL 52023921

diff --git a/Engine/Plugins/Experimental/UAF/UAFAnimNode/Source/UAFAnimNode/Public/UAF/AnimNodes/UAFBlendStack.h b/Engine/Plugins/Experimental/UAF/UAFAnimNode/Source/UAFAnimNode/Public/UAF/AnimNodes/UAFBlendStack.h
--- a/Engine/Plugins/Experimental/UAF/UAFAnimNode/Source/UAFAnimNode/Public/UAF/AnimNodes/UAFBlendStack.h
+++ b/Engine/Plugins/Experimental/UAF/UAFAnimNode/Source/UAFAnimNode/Public/UAF/AnimNodes/UAFBlendStack.h
@@ -37,7 +37,7 @@ public:
 	UAFANIMNODE_API void TransitionTo(FUAFAnimGraphUpdateContext& Context, FUAFAnimNode* TargetNodeData, const FUAFTransitionNodeData* TransitionData = nullptr);
 
 	// returns true if the BlendStack has a valid child node
-	UAFANIMNODE_API bool HasValidChild() const;
+	bool HasValidChild() const;
 		
 #if UAF_TRACE_ENABLED
 		virtual FString GetDebugName() const override;
@@ -45,9 +45,9 @@ public:
 #endif
 
 protected:
-	UAFANIMNODE_API const FUAFAnimNodePtr& GetChild() const;
-	UAFANIMNODE_API FUAFAnimNodePtr ReleaseChild();
-	UAFANIMNODE_API void SetChild(const FUAFAnimNodePtr& NewChild);
+	const FUAFAnimNodePtr& GetChild() const;
+	FUAFAnimNodePtr ReleaseChild();
+	void SetChild(const FUAFAnimNodePtr& NewChild);
 	
 	// if there is no valid child, add a Reference Pose Op, if there is a valid child clear Op
 	UAFANIMNODE_API void ValidateSingleChild();
diff --git a/Engine/Plugins/Experimental/UAF/UAFStateTree/Source/UAFStateTree/Internal/AnimNode/UAFStateTreeNode.h b/Engine/Plugins/Experimental/UAF/UAFStateTree/Source/UAFStateTree/Internal/AnimNode/UAFStateTreeNode.h
--- a/Engine/Plugins/Experimental/UAF/UAFStateTree/Source/UAFStateTree/Internal/AnimNode/UAFStateTreeNode.h
+++ b/Engine/Plugins/Experimental/UAF/UAFStateTree/Source/UAFStateTree/Internal/AnimNode/UAFStateTreeNode.h
@@ -56,17 +56,17 @@ struct FUAFStateTreeNodeData : public FUAFAnimNodeData
 	UAFSTATETREE_API virtual void* GetInterface(FUAFAnimNodeInterfaceId Id) override;
 
 	// HasAsset interface
-	UAFSTATETREE_API virtual UObject* GetAsset() const override
+	virtual UObject* GetAsset() const override
 	{
 		return StateTree;
 	}
 
-	UAFSTATETREE_API const UStateTree* GetStateTree() const
+	const UStateTree* GetStateTree() const
 	{
 		return StateTree ? StateTree->StateTree : nullptr;
 	}
 
-	UAFSTATETREE_API void SetStateTree(UAnimNextStateTree* InStateTree)
+	void SetStateTree(UAnimNextStateTree* InStateTree)
 	{
 		StateTree = InStateTree;
 	}
