Flip arguments to match common concentions

#Synced-CL 52066653

diff --git a/Engine/Plugins/Experimental/UAF/UAFAnimNode/Source/UAFAnimNode/Private/UAF/AnimNodeCore/UAFAnimNode.cpp b/Engine/Plugins/Experimental/UAF/UAFAnimNode/Source/UAFAnimNode/Private/UAF/AnimNodeCore/UAFAnimNode.cpp
--- a/Engine/Plugins/Experimental/UAF/UAFAnimNode/Source/UAFAnimNode/Private/UAF/AnimNodeCore/UAFAnimNode.cpp
+++ b/Engine/Plugins/Experimental/UAF/UAFAnimNode/Source/UAFAnimNode/Private/UAF/AnimNodeCore/UAFAnimNode.cpp
@@ -119,7 +119,7 @@ namespace UE::UAF
 		Children.RemoveAt(ChildIndex, EAllowShrinking::No);
 	}
 
-	void FUAFAnimNode::SetChildAt(const FUAFAnimNodePtr& Child, int32 ChildIndex)
+	void FUAFAnimNode::SetChildAt(int32 ChildIndex, const FUAFAnimNodePtr& Child)
 	{
 		checkf(!Child || !Children.Contains(Child), TEXT("Attempting to add a duplicate child entry."));
 
diff --git a/Engine/Plugins/Experimental/UAF/UAFAnimNode/Source/UAFAnimNode/Private/UAF/AnimNodes/UAFBlendStack.cpp b/Engine/Plugins/Experimental/UAF/UAFAnimNode/Source/UAFAnimNode/Private/UAF/AnimNodes/UAFBlendStack.cpp
--- a/Engine/Plugins/Experimental/UAF/UAFAnimNode/Source/UAFAnimNode/Private/UAF/AnimNodes/UAFBlendStack.cpp
+++ b/Engine/Plugins/Experimental/UAF/UAFAnimNode/Source/UAFAnimNode/Private/UAF/AnimNodes/UAFBlendStack.cpp
@@ -39,7 +39,7 @@ void FUAFBlendStack::PruneTransitions()
 				{
 					// Our transition is complete, prune it
 					FUAFAnimNodePtr Target = TransitionInstance->ReleaseTarget();
-					SetChildAt(Target, 0);
+					SetChildAt(0, Target);
 				}
 			}
 		}
diff --git a/Engine/Plugins/Experimental/UAF/UAFAnimNode/Source/UAFAnimNode/Private/UAF/AnimNodes/UAFSimpleTransition.cpp b/Engine/Plugins/Experimental/UAF/UAFAnimNode/Source/UAFAnimNode/Private/UAF/AnimNodes/UAFSimpleTransition.cpp
--- a/Engine/Plugins/Experimental/UAF/UAFAnimNode/Source/UAFAnimNode/Private/UAF/AnimNodes/UAFSimpleTransition.cpp
+++ b/Engine/Plugins/Experimental/UAF/UAFAnimNode/Source/UAFAnimNode/Private/UAF/AnimNodes/UAFSimpleTransition.cpp
@@ -83,7 +83,7 @@ namespace UE::UAF
 				TimeRemaining = 0.0f;
 
 				// Prune our source and stop blending
-				SetChildAt(nullptr, SourceChildIndex);
+				SetChildAt(SourceChildIndex, nullptr);
 				SetPostAnimOp(nullptr);
 
 				if (FUAFAnimNode* ParentNode = GetParent())
@@ -108,7 +108,7 @@ namespace UE::UAF
 							if (NestedTransition->IsComplete())
 							{
 								FUAFAnimNodePtr NewSource = NestedTransition->ReleaseTarget();
-								SetChildAt(NewSource, SourceChildIndex);
+								SetChildAt(SourceChildIndex, NewSource);
 							}
 						}
 					}
@@ -165,7 +165,7 @@ namespace UE::UAF
 		checkf(IsComplete(), TEXT("Attempting to release ownership of the transition source before it has completed."));
 
 		FUAFAnimNodePtr Source = GetChildAt(SourceChildIndex);
-		SetChildAt(nullptr, SourceChildIndex);
+		SetChildAt(SourceChildIndex, nullptr);
 
 		return Source;
 	}
@@ -175,7 +175,7 @@ namespace UE::UAF
 		checkf(IsComplete(), TEXT("Attempting to release ownership of the transition target before it has completed."));
 
 		FUAFAnimNodePtr Target = GetChildAt(TargetChildIndex);
-		SetChildAt(nullptr, TargetChildIndex);
+		SetChildAt(TargetChildIndex, nullptr);
 
 		return Target;
 	}
diff --git a/Engine/Plugins/Experimental/UAF/UAFAnimNode/Source/UAFAnimNode/Public/UAF/AnimNodeCore/UAFAnimNode.h b/Engine/Plugins/Experimental/UAF/UAFAnimNode/Source/UAFAnimNode/Public/UAF/AnimNodeCore/UAFAnimNode.h
--- a/Engine/Plugins/Experimental/UAF/UAFAnimNode/Source/UAFAnimNode/Public/UAF/AnimNodeCore/UAFAnimNode.h
+++ b/Engine/Plugins/Experimental/UAF/UAFAnimNode/Source/UAFAnimNode/Public/UAF/AnimNodeCore/UAFAnimNode.h
@@ -150,7 +150,7 @@ namespace UE::UAF
 		UE_API void RemoveChildAt(int32 ChildIndex);
 
 		// Sets the child at the specified index
-		UE_API void SetChildAt(const FUAFAnimNodePtr& Child, int32 ChildIndex);
+		UE_API void SetChildAt(int32 ChildIndex, const FUAFAnimNodePtr& Child);
 		
 		// Sets the AnimOp to perform before our children evaluate theirs
 		// When choosing whether to use PreAnimOp or PostAnimOp, favor the latter when possible
diff --git a/Engine/Plugins/Experimental/UAF/UAFAnimNode/Source/UAFAnimNode/Public/UAF/AnimNodes/UAFBlendStack.h b/Engine/Plugins/Experimental/UAF/UAFAnimNode/Source/UAFAnimNode/Public/UAF/AnimNodes/UAFBlendStack.h
--- a/Engine/Plugins/Experimental/UAF/UAFAnimNode/Source/UAFAnimNode/Public/UAF/AnimNodes/UAFBlendStack.h
+++ b/Engine/Plugins/Experimental/UAF/UAFAnimNode/Source/UAFAnimNode/Public/UAF/AnimNodes/UAFBlendStack.h
@@ -79,15 +79,14 @@ inline const FUAFAnimNodePtr& FUAFBlendStack::GetChild() const
 inline FUAFAnimNodePtr FUAFBlendStack::ReleaseChild()
 {
 	FUAFAnimNodePtr Child = GetFirstChild();
-	SetChildAt(nullptr, 0);
+	SetChildAt(0, nullptr);
 	return Child;
 }
 	
 inline void FUAFBlendStack::SetChild(const FUAFAnimNodePtr& NewChild)
 {
 	ensure(!GetChild().IsValid());
-	SetChildAt(NewChild,0);
+	SetChildAt(0, NewChild);
 }
-    	
 
 }
