Fix uaf anim node properties displaying on frames where the node wasn't active

#Synced-CL 51932823

diff --git a/Engine/Plugins/Experimental/UAF/UAFAnimNode/Source/UAFAnimNodeEditor/Private/RewindDebugger/UAFAnimNodeTrack.cpp b/Engine/Plugins/Experimental/UAF/UAFAnimNode/Source/UAFAnimNodeEditor/Private/RewindDebugger/UAFAnimNodeTrack.cpp
--- a/Engine/Plugins/Experimental/UAF/UAFAnimNode/Source/UAFAnimNodeEditor/Private/RewindDebugger/UAFAnimNodeTrack.cpp
+++ b/Engine/Plugins/Experimental/UAF/UAFAnimNode/Source/UAFAnimNodeEditor/Private/RewindDebugger/UAFAnimNodeTrack.cpp
@@ -120,10 +120,6 @@ TSharedPtr<SWidget> FUAFAnimNodeTrack::GetDetailsViewInternal()
         DetailsViewArgs.NameAreaSettings = FDetailsViewArgs::HideNameArea;
         DetailsView = PropertyEditorModule.CreateDetailView(DetailsViewArgs);
 		DetailsView->SetObject(DetailsObject.Get());
-		
-		static FName PropertyName("TotalWeight");
-		FPropertyBagPropertyDesc PropertyDesc(PropertyName,  EPropertyBagPropertyType::Float);
-		DetailsObject->Properties.AddProperties({PropertyDesc});
 	}
 
 	return DetailsView;
@@ -165,6 +161,9 @@ bool FUAFAnimNodeTrack::UpdateInternal()
 				{
 					PreviousScrubTime = CurrentScrubTime;
 
+					// clear properties first so we don't see old data
+					DetailsObject->Properties.Reset();
+
 					TraceServices::FFrame MarkerFrame;
 					if (FramesProvider.GetFrameFromTime(ETraceFrameType::TraceFrameType_Game, CurrentScrubTime, MarkerFrame))
 					{
@@ -174,6 +173,8 @@ bool FUAFAnimNodeTrack::UpdateInternal()
 							if (NodeUpdateData.NodeId == NodeId)
 							{
 								static FName PropertyName("TotalWeight");
+								FPropertyBagPropertyDesc PropertyDesc(PropertyName,  EPropertyBagPropertyType::Float);
+								DetailsObject->Properties.AddProperties({PropertyDesc});
 								DetailsObject->Properties.SetValueFloat(PropertyName, NodeUpdateData.TotalWeight);	
 								return TraceServices::EEventEnumerate::Stop;
 							}
