Add transition source when notifying a parent of a completed transition

#Synced-CL 52066653

diff --git a/Engine/Plugins/Experimental/UAF/UAFAnimNode/Source/UAFAnimNode/Private/UAF/AnimNodes/UAFBlendStack.cpp b/Engine/Plugins/Experimental/UAF/UAFAnimNode/Source/UAFAnimNode/Private/UAF/AnimNodes/UAFBlendStack.cpp
--- a/Engine/Plugins/Experimental/UAF/UAFAnimNode/Source/UAFAnimNode/Private/UAF/AnimNodes/UAFBlendStack.cpp
+++ b/Engine/Plugins/Experimental/UAF/UAFAnimNode/Source/UAFAnimNode/Private/UAF/AnimNodes/UAFBlendStack.cpp
@@ -125,7 +125,7 @@ void* FUAFBlendStack::GetInterface(FUAFAnimNodeInterfaceId Id)
 	return nullptr;
 }
 
-void FUAFBlendStack::NotifyTransitionComplete()
+void FUAFBlendStack::NotifyTransitionComplete(const IUAFTransitionNode& TransitionNode)
 {
 	bTransitionCompleted = true;
 }
diff --git a/Engine/Plugins/Experimental/UAF/UAFAnimNode/Source/UAFAnimNode/Private/UAF/AnimNodes/UAFSimpleTransition.cpp b/Engine/Plugins/Experimental/UAF/UAFAnimNode/Source/UAFAnimNode/Private/UAF/AnimNodes/UAFSimpleTransition.cpp
--- a/Engine/Plugins/Experimental/UAF/UAFAnimNode/Source/UAFAnimNode/Private/UAF/AnimNodes/UAFSimpleTransition.cpp
+++ b/Engine/Plugins/Experimental/UAF/UAFAnimNode/Source/UAFAnimNode/Private/UAF/AnimNodes/UAFSimpleTransition.cpp
@@ -62,7 +62,7 @@ namespace UE::UAF
 		}
 	}
 
-	void FUAFSimpleTransition::NotifyTransitionComplete()
+	void FUAFSimpleTransition::NotifyTransitionComplete(const IUAFTransitionNode& TransitionNode)
 	{
 		bNestedTransitionComplete = true;
 	}
@@ -91,7 +91,7 @@ namespace UE::UAF
 					// if our parent implements IUAFTransitionContainerNode notify it that we can be pruned
 					if (IUAFTransitionContainerNode* ParentContainer = ParentNode->GetInterface<IUAFTransitionContainerNode>())
 					{
-						ParentContainer->NotifyTransitionComplete();
+						ParentContainer->NotifyTransitionComplete(*this);
 					}
 				}
 			}
diff --git a/Engine/Plugins/Experimental/UAF/UAFAnimNode/Source/UAFAnimNode/Public/UAF/AnimNodeCore/IUAFTransitionContainerNode.h b/Engine/Plugins/Experimental/UAF/UAFAnimNode/Source/UAFAnimNode/Public/UAF/AnimNodeCore/IUAFTransitionContainerNode.h
--- a/Engine/Plugins/Experimental/UAF/UAFAnimNode/Source/UAFAnimNode/Public/UAF/AnimNodeCore/IUAFTransitionContainerNode.h
+++ b/Engine/Plugins/Experimental/UAF/UAFAnimNode/Source/UAFAnimNode/Public/UAF/AnimNodeCore/IUAFTransitionContainerNode.h
@@ -8,6 +8,8 @@
 
 namespace UE::UAF
 {
+	struct IUAFTransitionNode;
+
 	/**
 	 * IUAFTransitionContainerNode
 	 *
@@ -20,7 +22,7 @@ namespace UE::UAF
 		virtual ~IUAFTransitionContainerNode() = default;
 
 		// Called by transitions when they complete
-		virtual void NotifyTransitionComplete() = 0;
+		virtual void NotifyTransitionComplete(const IUAFTransitionNode& TransitionNode) = 0;
 	};
 	
 }
diff --git a/Engine/Plugins/Experimental/UAF/UAFAnimNode/Source/UAFAnimNode/Public/UAF/AnimNodes/UAFBlendStack.h b/Engine/Plugins/Experimental/UAF/UAFAnimNode/Source/UAFAnimNode/Public/UAF/AnimNodes/UAFBlendStack.h
--- a/Engine/Plugins/Experimental/UAF/UAFAnimNode/Source/UAFAnimNode/Public/UAF/AnimNodes/UAFBlendStack.h
+++ b/Engine/Plugins/Experimental/UAF/UAFAnimNode/Source/UAFAnimNode/Public/UAF/AnimNodes/UAFBlendStack.h
@@ -31,7 +31,7 @@ public:
 	UAFANIMNODE_API virtual void* GetInterface(FUAFAnimNodeInterfaceId Id) override;
 
 	// IUAFTransitionContainerNode interface
-	UAFANIMNODE_API virtual void NotifyTransitionComplete() override;
+	UAFANIMNODE_API virtual void NotifyTransitionComplete(const IUAFTransitionNode& TransitionNode) override;
 
 	// Trigger blend to new node instance
 	UAFANIMNODE_API void TransitionTo(FUAFAnimGraphUpdateContext& Context, FUAFAnimNode* TargetNodeData, const FUAFTransitionNodeData* TransitionData = nullptr);
diff --git a/Engine/Plugins/Experimental/UAF/UAFAnimNode/Source/UAFAnimNode/Public/UAF/AnimNodes/UAFSimpleTransition.h b/Engine/Plugins/Experimental/UAF/UAFAnimNode/Source/UAFAnimNode/Public/UAF/AnimNodes/UAFSimpleTransition.h
--- a/Engine/Plugins/Experimental/UAF/UAFAnimNode/Source/UAFAnimNode/Public/UAF/AnimNodes/UAFSimpleTransition.h
+++ b/Engine/Plugins/Experimental/UAF/UAFAnimNode/Source/UAFAnimNode/Public/UAF/AnimNodes/UAFSimpleTransition.h
@@ -47,14 +47,14 @@ namespace UE::UAF
 		UE_API FUAFSimpleTransition(FUAFAnimGraphUpdateContext& Context, FUAFAnimNodePtr InSource, FUAFAnimNodePtr InTarget, float InDuration, EAlphaBlendOption InBlendOption);
 
 		// ITransitionInstance impl
-		UE_API virtual const FUAFAnimNodePtr& GetSource() const override;
-		UE_API virtual const FUAFAnimNodePtr& GetTarget() const override;
+		virtual const FUAFAnimNodePtr& GetSource() const override;
+		virtual const FUAFAnimNodePtr& GetTarget() const override;
 		UE_API virtual FUAFAnimNodePtr ReleaseSource() override;
 		UE_API virtual FUAFAnimNodePtr ReleaseTarget() override;
-		UE_API virtual bool IsComplete() const override;
+		virtual bool IsComplete() const override;
 
 		// ITransitionContainer impl
-		UE_API virtual void NotifyTransitionComplete() override;
+		UE_API virtual void NotifyTransitionComplete(const IUAFTransitionNode& TransitionNode) override;
 
 		// FUAFAnimNode impl
 		UE_API virtual void PreUpdate(FUAFAnimGraphUpdateContext& GraphContext) override;
