Add FUAFBlendProfileTransition
* Update FUAFSimpleTransition to derive UUAFTimedTransition (same as FUAFBlendProfleTransition) to reuse functionality.
* Add blend option to FUAFBlendTwoAnimOp.

#Synced-CL 52140028

diff --git a/Engine/Plugins/Experimental/UAF/UAFAnimNode/Source/UAFAnimNode/Private/UAF/AnimNodes/UAFBlendProfileTransition.cpp b/Engine/Plugins/Experimental/UAF/UAFAnimNode/Source/UAFAnimNode/Private/UAF/AnimNodes/UAFBlendProfileTransition.cpp
new file mode 100644
--- /dev/null
+++ b/Engine/Plugins/Experimental/UAF/UAFAnimNode/Source/UAFAnimNode/Private/UAF/AnimNodes/UAFBlendProfileTransition.cpp
@@ -0,0 +1,71 @@
+// Copyright Epic Games, Inc. All Rights Reserved.
+
+#include "UAF/AnimNodes/UAFBlendProfileTransition.h"
+
+#include "UAF/BlendProfile/UAFBlendProfile.h"
+
+namespace UE::UAF
+{
+	FUAFAnimNodePtr FUAFBlendProfileTransitionData::CreateInstance(FUAFAnimGraphUpdateContext& Context, FUAFAnimNodePtr Source, FUAFAnimNodePtr Target) const
+	{
+		if (Duration <= 0.0f)
+		{
+			return Target;
+		}
+
+		if (!BlendProfile)
+		{
+			return MakeAnimNode<FUAFSimpleTransition>(Context, Source, Target, Duration, BlendOption);
+		}
+
+		return MakeAnimNode<FUAFBlendProfileTransition>(Context, Source, Target, BlendProfile.Get(), Duration, BlendOption);
+	}
+
+	FUAFBlendProfileTransition::FUAFBlendProfileTransition(
+		FUAFAnimGraphUpdateContext& Context,
+		FUAFAnimNodePtr InSource,
+		FUAFAnimNodePtr InTarget,
+		TNonNullPtr<UUAFBlendProfile> InBlendProfile,
+		float InBaseDuration,
+		EAlphaBlendOption InBlendOption)
+		: FUAFTimedTransition(Context, InSource, InTarget, InBaseDuration, InBlendOption)
+		, BlendAnimOp(InBlendProfile, InBlendOption)
+	{
+		checkf(InSource || InTarget, TEXT("Must have a source and/or target to generate a transition"));
+		checkf(InSource != InTarget, TEXT("Cannot transition to/from the same nodes"));
+		checkf(InBaseDuration > 0.0f, TEXT("Profile transition must have a positive duration"));
+		InitializeAs<FUAFBlendProfileTransition>(Context);
+
+		BlendAnimOp.SetDebugOwner(this);
+
+		SetPostAnimOp(&BlendAnimOp);
+	}
+
+	void FUAFBlendProfileTransition::PreUpdate(FUAFAnimGraphUpdateContext& GraphContext)
+	{
+		SCOPED_NAMED_EVENT(AnimNode_Update_UAFProfileTransitionNode, FColor::Blue);
+
+		FUAFTimedTransition::PreUpdate(GraphContext);
+		
+		if (!IsComplete())
+		{
+			const float SourceWeight = TimeRemaining / Duration;
+			const float TargetWeight = 1.0f - SourceWeight;
+
+			BlendAnimOp.SetLinearWeightB(TargetWeight);
+		}
+	}
+
+#if UAF_TRACE_ENABLED
+	FString FUAFBlendProfileTransition::GetDebugName() const
+	{
+		static FString ProfileTransitionName("Profile Transition");
+		return ProfileTransitionName;
+	}
+
+	UStruct* FUAFBlendProfileTransition::GetDebugStruct() const
+	{
+		return FUAFSimpleTransitionData::StaticStruct();
+	}
+#endif
+}
diff --git a/Engine/Plugins/Experimental/UAF/UAFAnimNode/Source/UAFAnimNode/Private/UAF/AnimNodes/UAFSimpleTransition.cpp b/Engine/Plugins/Experimental/UAF/UAFAnimNode/Source/UAFAnimNode/Private/UAF/AnimNodes/UAFSimpleTransition.cpp
--- a/Engine/Plugins/Experimental/UAF/UAFAnimNode/Source/UAFAnimNode/Private/UAF/AnimNodes/UAFSimpleTransition.cpp
+++ b/Engine/Plugins/Experimental/UAF/UAFAnimNode/Source/UAFAnimNode/Private/UAF/AnimNodes/UAFSimpleTransition.cpp
@@ -20,133 +20,34 @@ namespace UE::UAF
 	}
 
 	FUAFSimpleTransition::FUAFSimpleTransition(FUAFAnimGraphUpdateContext& Context, FUAFAnimNodePtr InSource, FUAFAnimNodePtr InTarget, float InDuration, EAlphaBlendOption InBlendOption)
-		: FUAFAnimNode(Context)
-		, Duration(InDuration)
-		, BlendOption(InBlendOption)
-		, TimeRemaining(InDuration)
+		: FUAFTimedTransition(Context, InSource, InTarget, InDuration, InBlendOption)
 	{
-		checkf(InSource || InTarget, TEXT("Must have a source and/or target to generate a transition"));
 		InitializeAs<FUAFSimpleTransition>(Context);
 
 		if (InSource && InTarget && InDuration > 0.0f)
 		{
-			AddChild(InSource);
-			AddChild(InTarget);
-
-			InSource->SetIsBlendingOut(true);
-
 			BlendAnimOp.SetDebugOwner(this);
-
+			BlendAnimOp.SetBlendOption(InBlendOption);
+			
 			SetPostAnimOp(&BlendAnimOp);
 		}
-		else if (InTarget)
-		{
-			// We only have a target or our duration is zero, we'll snap to it
-			AddChild(nullptr);
-			AddChild(InTarget);
-
-			Duration = 0.0f;
-			TimeRemaining = 0.0f;
-		}
-		else
-		{
-			// We have no target, we cannot transition meaningfully
-			// We do not track the source, we'll pop to the bind pose
-			AddChild(nullptr);
-			AddChild(nullptr);
-
-			SetPostAnimOp(FUAFNullAnimOp::Get());
-
-			Duration = 0.0f;
-			TimeRemaining = 0.0f;
-		}
-	}
-
-	void FUAFSimpleTransition::NotifyTransitionComplete(const IUAFTransitionNode& TransitionNode)
-	{
-		bNestedTransitionComplete = true;
 	}
 
 	void FUAFSimpleTransition::PreUpdate(FUAFAnimGraphUpdateContext& GraphContext)
 	{
 		SCOPED_NAMED_EVENT(AnimNode_Update_UAFBasocTransitionNode, FColor::Blue);
 
+		FUAFTimedTransition::PreUpdate(GraphContext);
+		
 		if (!IsComplete())
 		{
-			const float CurrentWeight = GetTotalWeight();
-
-			TimeRemaining -= GraphContext.GetDeltaTime();
-
-			if (TimeRemaining <= 0.0f)
-			{
-				// Transition is done
-				TimeRemaining = 0.0f;
-
-				// Prune our source and stop blending
-				SetChildAt(SourceChildIndex, nullptr);
-				SetPostAnimOp(nullptr);
-
-				if (FUAFAnimNode* ParentNode = GetParent())
-				{
-					// if our parent implements IUAFTransitionContainerNode notify it that we can be pruned
-					if (IUAFTransitionContainerNode* ParentContainer = ParentNode->GetInterface<IUAFTransitionContainerNode>())
-					{
-						ParentContainer->NotifyTransitionComplete(*this);
-					}
-				}
-			}
-			else
-			{
-				// Nested transition has notified us that it is complete
-				if (bNestedTransitionComplete)
-				{
-					if (FUAFAnimNode* SourceChild = GetChildAt(SourceChildIndex))
-					{
-						// If we have a nested transition as our source, prune it when it completes
-						if (IUAFTransitionNode* NestedTransition = SourceChild->GetInterface<IUAFTransitionNode>())
-						{
-							if (NestedTransition->IsComplete())
-							{
-								FUAFAnimNodePtr NewSource = NestedTransition->ReleaseTarget();
-								SetChildAt(SourceChildIndex, NewSource);
-							}
-						}
-					}
-					bNestedTransitionComplete = false;
-				}
-
-				const float SourceWeight = TimeRemaining / Duration;
-				const float TargetWeight = 1.0f - SourceWeight;
-
-				UAF_TRACE_ANIMNODE_VALUE(GraphContext, this, "TransitionTargetWeight", TargetWeight);
-
-				GetChildAt(SourceChildIndex)->SetTotalWeight(CurrentWeight * SourceWeight);
-				GetChildAt(TargetChildIndex)->SetTotalWeight(CurrentWeight * TargetWeight);
-
-				BlendAnimOp.SetInterpolationAlpha(TargetWeight);
-			}
+			const float SourceWeight = TimeRemaining / Duration;
+			const float TargetWeight = 1.0f - SourceWeight;
+			
+			BlendAnimOp.SetInterpolationAlpha(TargetWeight);
 		}
 	}
 
-	void* FUAFSimpleTransition::GetInterface(FUAFAnimNodeInterfaceId Id)
-	{
-		if (Id == IUAFTransitionNode::InterfaceId)
-		{
-			return static_cast<IUAFTransitionNode*>(this);
-		}
-		if (Id == IUAFTransitionContainerNode::InterfaceId)
-		{
-			return static_cast<IUAFTransitionContainerNode*>(this);
-		}
-
-		if (const FUAFAnimNodePtr& Target = GetChildAt(TargetChildIndex))
-		{
-			return Target->GetInterface(Id);
-		}
-
-		return nullptr;
-	}
-
 #if UAF_TRACE_ENABLED
 	FString FUAFSimpleTransition::GetDebugName() const
 	{
@@ -160,23 +61,4 @@ namespace UE::UAF
 	}
 #endif
 
-	FUAFAnimNodePtr FUAFSimpleTransition::ReleaseSource()
-	{
-		checkf(IsComplete(), TEXT("Attempting to release ownership of the transition source before it has completed."));
-
-		FUAFAnimNodePtr Source = GetChildAt(SourceChildIndex);
-		SetChildAt(SourceChildIndex, nullptr);
-
-		return Source;
-	}
-
-	FUAFAnimNodePtr FUAFSimpleTransition::ReleaseTarget()
-	{
-		checkf(IsComplete(), TEXT("Attempting to release ownership of the transition target before it has completed."));
-
-		FUAFAnimNodePtr Target = GetChildAt(TargetChildIndex);
-		SetChildAt(TargetChildIndex, nullptr);
-
-		return Target;
-	}
 }
diff --git a/Engine/Plugins/Experimental/UAF/UAFAnimNode/Source/UAFAnimNode/Private/UAF/AnimNodes/UAFTimedTransition.cpp b/Engine/Plugins/Experimental/UAF/UAFAnimNode/Source/UAFAnimNode/Private/UAF/AnimNodes/UAFTimedTransition.cpp
new file mode 100644
--- /dev/null
+++ b/Engine/Plugins/Experimental/UAF/UAFAnimNode/Source/UAFAnimNode/Private/UAF/AnimNodes/UAFTimedTransition.cpp
@@ -0,0 +1,166 @@
+// Copyright Epic Games, Inc. All Rights Reserved.
+
+#include "UAF/AnimNodes/UAFTimedTransition.h"
+
+#include "UAF/AnimNodeCore/UAFAnimNodeUpdate.h"
+#include "UAF/AnimNodeCore/UAFTransitionNodeData.h"
+#include "UAF/AnimOps/UAFNullAnimOp.h"
+
+namespace UE::UAF
+{
+	FUAFTimedTransition::FUAFTimedTransition(FUAFAnimGraphUpdateContext& Context, FUAFAnimNodePtr InSource, FUAFAnimNodePtr InTarget, float InDuration, EAlphaBlendOption InBlendOption)
+		: FUAFAnimNode(Context)
+		, Duration(InDuration)
+		, BlendOption(InBlendOption)
+		, TimeRemaining(InDuration)
+	{
+		checkf(InSource || InTarget, TEXT("Must have a source and/or target to generate a transition"));
+		InitializeAs<FUAFTimedTransition>(Context);
+
+		SetPostAnimOp(nullptr);
+		
+		if (InSource && InTarget && InDuration > 0.0f)
+		{
+			AddChild(InSource);
+			AddChild(InTarget);
+
+			InSource->SetIsBlendingOut(true);
+		}
+		else if (InTarget)
+		{
+			// We only have a target or our duration is zero, we'll snap to it
+			AddChild(nullptr);
+			AddChild(InTarget);
+
+			Duration = 0.0f;
+			TimeRemaining = 0.0f;
+		}
+		else
+		{
+			// We have no target, we cannot transition meaningfully
+			// We do not track the source, we'll pop to the bind pose
+			AddChild(nullptr);
+			AddChild(nullptr);
+
+			Duration = 0.0f;
+			TimeRemaining = 0.0f;
+		}
+	}
+
+	void FUAFTimedTransition::NotifyTransitionComplete(const IUAFTransitionNode& TransitionNode)
+	{
+		bNestedTransitionComplete = true;
+	}
+
+	FUAFAnimNodePtr FUAFTimedTransition::ReleaseSource()
+	{
+		checkf(IsComplete(), TEXT("Attempting to release ownership of the transition source before it has completed."));
+
+		FUAFAnimNodePtr Source = GetChildAt(SourceChildIndex);
+		SetChildAt(SourceChildIndex, nullptr);
+
+		return Source;
+	}
+
+	FUAFAnimNodePtr FUAFTimedTransition::ReleaseTarget()
+	{
+		checkf(IsComplete(), TEXT("Attempting to release ownership of the transition target before it has completed."));
+
+		FUAFAnimNodePtr Target = GetChildAt(TargetChildIndex);
+		SetChildAt(TargetChildIndex, nullptr);
+
+		return Target;
+	}
+
+	void FUAFTimedTransition::PreUpdate(FUAFAnimGraphUpdateContext& GraphContext)
+	{
+		SCOPED_NAMED_EVENT(AnimNode_Update_UAFBasocTransitionNode, FColor::Blue);
+
+		if (!IsComplete())
+		{
+			const float CurrentWeight = GetTotalWeight();
+
+			TimeRemaining -= GraphContext.GetDeltaTime();
+
+			if (TimeRemaining <= 0.0f)
+			{
+				// Transition is done
+				TimeRemaining = 0.0f;
+
+				// Prune our source and stop blending
+				SetChildAt(SourceChildIndex, nullptr);
+				SetPostAnimOp(nullptr);
+
+				if (FUAFAnimNode* ParentNode = GetParent())
+				{
+					// if our parent implements IUAFTransitionContainerNode notify it that we can be pruned
+					if (IUAFTransitionContainerNode* ParentContainer = ParentNode->GetInterface<IUAFTransitionContainerNode>())
+					{
+						ParentContainer->NotifyTransitionComplete(*this);
+					}
+				}
+			}
+			else
+			{
+				// Nested transition has notified us that it is complete
+				if (bNestedTransitionComplete)
+				{
+					if (FUAFAnimNode* SourceChild = GetChildAt(SourceChildIndex))
+					{
+						// If we have a nested transition as our source, prune it when it completes
+						if (IUAFTransitionNode* NestedTransition = SourceChild->GetInterface<IUAFTransitionNode>())
+						{
+							if (NestedTransition->IsComplete())
+							{
+								FUAFAnimNodePtr NewSource = NestedTransition->ReleaseTarget();
+								SetChildAt(SourceChildIndex, NewSource);
+							}
+						}
+					}
+					bNestedTransitionComplete = false;
+				}
+
+				const float SourceWeight = TimeRemaining / Duration;
+				const float TargetWeight = 1.0f - SourceWeight;
+
+				UAF_TRACE_ANIMNODE_VALUE(GraphContext, this, "TransitionTargetWeight", TargetWeight);
+
+				GetChildAt(SourceChildIndex)->SetTotalWeight(CurrentWeight * SourceWeight);
+				GetChildAt(TargetChildIndex)->SetTotalWeight(CurrentWeight * TargetWeight);
+			}
+		}	
+	}
+	
+	void* FUAFTimedTransition::GetInterface(FUAFAnimNodeInterfaceId Id)
+	{
+		if (Id == IUAFTransitionNode::InterfaceId)
+		{
+			return static_cast<IUAFTransitionNode*>(this);
+		}
+		if (Id == IUAFTransitionContainerNode::InterfaceId)
+		{
+			return static_cast<IUAFTransitionContainerNode*>(this);
+		} 
+
+		if (const FUAFAnimNodePtr& Target = GetChildAt(TargetChildIndex))
+		{
+			return Target->GetInterface(Id);
+		}
+
+		return nullptr;
+	}
+	
+#if UAF_TRACE_ENABLED
+	FString FUAFTimedTransition::GetDebugName() const
+	{
+		static FString TimedTransitionName("Timed Transition");
+		return TimedTransitionName;
+	}
+
+	UStruct* FUAFTimedTransition::GetDebugStruct() const
+	{
+		return FUAFTransitionNodeData::StaticStruct();
+	}
+#endif
+	
+}
\ No newline at end of file
diff --git a/Engine/Plugins/Experimental/UAF/UAFAnimNode/Source/UAFAnimNode/Private/UAF/AnimOps/UAFBlendTwoAnimOp.cpp b/Engine/Plugins/Experimental/UAF/UAFAnimNode/Source/UAFAnimNode/Private/UAF/AnimOps/UAFBlendTwoAnimOp.cpp
--- a/Engine/Plugins/Experimental/UAF/UAFAnimNode/Source/UAFAnimNode/Private/UAF/AnimOps/UAFBlendTwoAnimOp.cpp
+++ b/Engine/Plugins/Experimental/UAF/UAFAnimNode/Source/UAFAnimNode/Private/UAF/AnimOps/UAFBlendTwoAnimOp.cpp
@@ -28,6 +28,11 @@ namespace UE::UAF
 		InputScaleBiasClampFn = MoveTemp(InInputScaleBiasClampFn);
 	}
 
+	void FUAFBlendTwoAnimOp::SetBlendOption(EAlphaBlendOption InBlendOption)
+	{
+		BlendOption = InBlendOption;
+	}
+	
 	void FUAFBlendTwoAnimOp::EvaluateValues(FUAFAnimOpValueEvaluator& Evaluator)
 	{
 		const FUAFAnimOpValueEvaluationContext& AnimOpEvaluationContext = Evaluator.GetActiveEvaluationContext();
@@ -126,6 +131,8 @@ namespace UE::UAF
 			}
 		}
 
+		Alpha = FAlphaBlend::AlphaToBlendOption(Alpha, BlendOption);
+
 		return FMath::Clamp(Alpha, 0.0f, 1.0f);
 	}
 }
diff --git a/Engine/Plugins/Experimental/UAF/UAFAnimNode/Source/UAFAnimNode/Private/UAF/AnimOps/UAFBlendTwoPerValueAnimOp.cpp b/Engine/Plugins/Experimental/UAF/UAFAnimNode/Source/UAFAnimNode/Private/UAF/AnimOps/UAFBlendTwoPerValueAnimOp.cpp
new file mode 100644
--- /dev/null
+++ b/Engine/Plugins/Experimental/UAF/UAFAnimNode/Source/UAFAnimNode/Private/UAF/AnimOps/UAFBlendTwoPerValueAnimOp.cpp
@@ -0,0 +1,164 @@
+// Copyright Epic Games, Inc. All Rights Reserved.
+
+#include "UAF/AnimOps/UAFBlendTwoPerValueAnimOp.h"
+
+#include "Animation/InputScaleBias.h"
+#include "UAF/AnimOpCore/UAFAnimOpValueEvaluator.h"
+#include "UAF/BlendProfile/UAFBlendProfile.h"
+#include "UAF/ValueRuntime/Transformers/Interpolate.h"
+#include "UAF/ValueRuntime/Transformers/Sanitize.h"
+#include "UAF/ValueRuntime/ValueBundle.h"
+
+namespace UE::UAF
+{
+	FUAFBlendTwoPerValueAnimOp::FUAFBlendTwoPerValueAnimOp()
+		: FUAFAnimOp(2)
+		, BlendProfile(nullptr)
+	{
+		InitializeAs<FUAFBlendTwoPerValueAnimOp>();
+	}
+	
+	FUAFBlendTwoPerValueAnimOp::FUAFBlendTwoPerValueAnimOp(TNonNullPtr<UUAFBlendProfile> InBlendProfile, EAlphaBlendOption InBlendOption)
+		: FUAFAnimOp(2)
+		, BlendProfile(InBlendProfile)
+		, BlendOption(InBlendOption)
+	{
+		InitializeAs<FUAFBlendTwoPerValueAnimOp>();
+	}
+
+	void FUAFBlendTwoPerValueAnimOp::SetLinearWeightB(float InInterpolationAlpha)
+	{
+		LinearGlobalWeightB = FMath::Clamp(InInterpolationAlpha, 0.0f, 1.0f);
+	}
+
+	void FUAFBlendTwoPerValueAnimOp::EvaluateValues(FUAFAnimOpValueEvaluator& Evaluator)
+	{
+		const FUAFAnimOpValueEvaluationContext& AnimOpEvaluationContext = Evaluator.GetActiveEvaluationContext();
+
+		// Pop our inputs
+		FValueBundleCoWRef InputRefB = Evaluator.GetEvaluationStack().Pop();
+		FValueBundleCoWRef InputRefA = Evaluator.GetEvaluationStack().Pop();
+
+		const FValueBundle& InputA = InputRefA.Get();
+		const FValueBundle& InputB = InputRefB.Get();
+
+		const float GlobalWeightB = FAlphaBlend::AlphaToBlendOption(LinearGlobalWeightB, BlendOption);
+		const float GlobalWeightA = 1.0 - GlobalWeightB;
+
+		if (!FAnimWeight::IsRelevant(GlobalWeightB))
+		{
+			// Input A has full weight, push it back even if it might be empty
+			Evaluator.GetEvaluationStack().Push(MoveTemp(InputRefA));
+		}
+		else if (FAnimWeight::IsFullWeight(GlobalWeightB))
+		{
+			// Input B has full weight, push it back even if it might be empty
+			Evaluator.GetEvaluationStack().Push(MoveTemp(InputRefB));
+		}
+		else
+		{
+			// TODO: It would be nice if we could use an immutable reference to the bind pose to avoid copying it
+
+			if (InputA.IsEmpty())
+			{
+				if (InputB.IsEmpty())
+				{
+					// Both inputs are empty, return an empty output
+					Evaluator.GetEvaluationStack().Push(MoveTemp(InputRefB));
+					return;
+				}
+				else
+				{
+					// A is empty, initialize it with the same space as B
+					check(InputRefA.IsMutable());
+					InputRefA.GetMutable().InitWithValueSpace(InputB.GetValueSpace());
+				}
+			}
+			else
+			{
+				if (InputB.IsEmpty())
+				{
+					// B is empty, initialize it with the same space as A
+					check(InputRefB.IsMutable());
+					InputRefB.GetMutable().InitWithValueSpace(InputA.GetValueSpace());
+				}
+			}
+
+			// Named sets must match as it ensures that our inputs have the same sizes/shapes
+			check(InputA.GetNamedSet() == InputB.GetNamedSet());
+
+			FValueBundleCoWRef& OutputRef = FindMutableCoWRef(InputRefB, InputRefA);
+
+			FValueBundleStack PerValueWeights(Evaluator.GetActiveNamedSet());
+			PerValueWeights.InitWithValueSpace(FValueSpace(EValueSpaceType::Local));
+			{
+				FAttributeMappingKey MappingKey = FAttributeMappingKey::MakeFromTo<FBoneTransformAnimationAttribute, FFloatAnimationAttribute>();
+				TBoundValueMap<FFloatAnimationAttribute>* BoneWeightsMap = PerValueWeights.GetBoundValueMaps().Add<FFloatAnimationAttribute>(MappingKey);
+
+				// TODO: In the future we don't want to be using USkeleton backed blend profiles and instead should use blend profiles authored
+				// for the abstract skeleton but in the meantime we can do this shim even though its a bit slow
+				const UUAFBlendProfile::FSkeletonBoneWeightArray& ProfileBoneWeights = BlendProfile->GetSkeletonBoneWeights();
+				BoneWeightsMap->FillWithValue(FFloatAnimationAttribute{ GlobalWeightB });
+
+				const FAttributeTypedSetPtr& BoneTransforms = InputA.GetNamedSet()->FindTypedSet<FBoneTransformAnimationAttribute>();
+				const FReferenceSkeleton& RefSkeleton = BlendProfile->GetSkeleton()->GetReferenceSkeleton();
+				
+				if (BlendProfile->GetType() == EUAFBlendProfileType::TimeFactor)
+				{
+					// For time profiles we scale the global weight by the profile value
+					// and then sample the blend curve.
+					for (int32 BoneIndex = 0; BoneIndex < ProfileBoneWeights.Num(); ++BoneIndex)
+					{
+						const FName BoneName = RefSkeleton.GetBoneName(BoneIndex);
+
+						FAttributeSetKey Key = BoneTransforms->FindKey(BoneName); // O(log n)
+						if (Key.IsValid())
+						{
+							const float ProfileValue = ProfileBoneWeights[BoneIndex];
+							check(ProfileValue >= 0.0f && ProfileValue <= 1.0f);
+							
+							const float BoneLinearWeightB = ProfileValue != 0.0f
+								? FMath::Clamp(LinearGlobalWeightB / ProfileValue, 0.0f, 1.0f)
+								: 1.0f;
+							
+							const float BoneWeightB = FAlphaBlend::AlphaToBlendOption(BoneLinearWeightB, BlendOption); // Convert linear weight to curve weight
+
+							FAttributeSetIndex Index = BoneTransforms->FindIndex(Key); // O(log n)
+							BoneWeightsMap->SetValue(Index, FFloatAnimationAttribute{ BoneWeightB });
+						}
+					}
+				}
+				else
+				{
+					// For weight profiles we sample the blend curve using the global weight and then
+					// scale it by the profile value in both directions and normalize the result.
+					for (int32 BoneIndex = 0; BoneIndex < ProfileBoneWeights.Num(); ++BoneIndex)
+					{
+						const FName BoneName = RefSkeleton.GetBoneName(BoneIndex);
+
+						FAttributeSetKey Key = BoneTransforms->FindKey(BoneName); // O(log n)
+						if (Key.IsValid())
+						{
+							const float ProfileValue = ProfileBoneWeights[BoneIndex];
+							check(ProfileValue >= 1.0f);
+							
+							float BoneWeightA = GlobalWeightA / ProfileValue;
+							float BoneWeightB = GlobalWeightB * ProfileValue;
+
+							const float SumBoneWeights = BoneWeightA + BoneWeightB;
+							BoneWeightB /= SumBoneWeights;
+
+							FAttributeSetIndex Index = BoneTransforms->FindIndex(Key); // O(log n)
+							BoneWeightsMap->SetValue(Index, FFloatAnimationAttribute{ BoneWeightB });
+						}
+					}
+				}
+			}
+			
+			Transformers::FInterpolate::Apply(Evaluator.GetTransformerMap(), InputA, InputB, PerValueWeights, LinearGlobalWeightB, OutputRef.GetMutable());
+			Transformers::FSanitize::Apply(Evaluator.GetTransformerMap(), OutputRef.GetMutable());
+
+			Evaluator.GetEvaluationStack().Push(MoveTemp(OutputRef));
+		}
+	}
+}
\ No newline at end of file
diff --git a/Engine/Plugins/Experimental/UAF/UAFAnimNode/Source/UAFAnimNode/Public/UAF/AnimNodeCore/UAFTransitionNodeData.h b/Engine/Plugins/Experimental/UAF/UAFAnimNode/Source/UAFAnimNode/Public/UAF/AnimNodeCore/UAFTransitionNodeData.h
--- a/Engine/Plugins/Experimental/UAF/UAFAnimNode/Source/UAFAnimNode/Public/UAF/AnimNodeCore/UAFTransitionNodeData.h
+++ b/Engine/Plugins/Experimental/UAF/UAFAnimNode/Source/UAFAnimNode/Public/UAF/AnimNodeCore/UAFTransitionNodeData.h
@@ -16,7 +16,7 @@ namespace UE::UAF
 	 *
 	 * Base structure for blend transition shared data
 	 */
-	USTRUCT()
+	USTRUCT(meta = (Hidden))
 	struct FUAFTransitionNodeData
 	{
 		GENERATED_BODY()
diff --git a/Engine/Plugins/Experimental/UAF/UAFAnimNode/Source/UAFAnimNode/Public/UAF/AnimNodes/UAFBlendProfileTransition.h b/Engine/Plugins/Experimental/UAF/UAFAnimNode/Source/UAFAnimNode/Public/UAF/AnimNodes/UAFBlendProfileTransition.h
new file mode 100644
--- /dev/null
+++ b/Engine/Plugins/Experimental/UAF/UAFAnimNode/Source/UAFAnimNode/Public/UAF/AnimNodes/UAFBlendProfileTransition.h
@@ -0,0 +1,54 @@
+// Copyright Epic Games, Inc. All Rights Reserved.
+
+#pragma once
+
+#include "CoreMinimal.h"
+#include "UAFSimpleTransition.h"
+#include "UAF/AnimOps/UAFBlendTwoPerValueAnimOp.h"
+
+#include "UAFBlendProfileTransition.generated.h"
+
+#define UE_API UAFANIMNODE_API
+
+class UUAFBlendProfile;
+
+namespace UE::UAF
+{
+	USTRUCT(DisplayName = "Blend Profile Transition")
+	struct FUAFBlendProfileTransitionData : public FUAFSimpleTransitionData
+	{
+		GENERATED_BODY()
+
+		UPROPERTY(EditAnywhere, Category = Blending)
+		TObjectPtr<UUAFBlendProfile> BlendProfile;
+
+		// FUAFTransitionNodeData impl
+		virtual FUAFAnimNodePtr CreateInstance(FUAFAnimGraphUpdateContext& Context, FUAFAnimNodePtr Source, FUAFAnimNodePtr Target) const override;
+	};
+
+	class FUAFBlendProfileTransition : public FUAFTimedTransition
+	{
+	public:
+		FUAFBlendProfileTransition() = delete;
+		UE_API FUAFBlendProfileTransition(
+			FUAFAnimGraphUpdateContext& Context,
+			FUAFAnimNodePtr InSource,
+			FUAFAnimNodePtr InTarget,
+			TNonNullPtr<UUAFBlendProfile> InBlendProfile,
+			float InBaseDuration,
+			EAlphaBlendOption InBlendOption);
+
+		// FUAFAnimNode impl
+		UE_API virtual void PreUpdate(FUAFAnimGraphUpdateContext& GraphContext) override;
+
+#if UAF_TRACE_ENABLED
+		UE_API virtual FString GetDebugName() const override;
+		UE_API virtual UStruct* GetDebugStruct() const override;
+#endif
+		
+	protected:
+		FUAFBlendTwoPerValueAnimOp BlendAnimOp;
+	};
+}
+
+#undef UE_API
\ No newline at end of file
diff --git a/Engine/Plugins/Experimental/UAF/UAFAnimNode/Source/UAFAnimNode/Public/UAF/AnimNodes/UAFSimpleTransition.h b/Engine/Plugins/Experimental/UAF/UAFAnimNode/Source/UAFAnimNode/Public/UAF/AnimNodes/UAFSimpleTransition.h
--- a/Engine/Plugins/Experimental/UAF/UAFAnimNode/Source/UAFAnimNode/Public/UAF/AnimNodes/UAFSimpleTransition.h
+++ b/Engine/Plugins/Experimental/UAF/UAFAnimNode/Source/UAFAnimNode/Public/UAF/AnimNodes/UAFSimpleTransition.h
@@ -4,6 +4,7 @@
 
 #include "CoreMinimal.h"
 #include "AlphaBlend.h"
+#include "UAFTimedTransition.h"
 #include "UAF/AnimNodeCore/UAFAnimNode.h"
 #include "UAF/AnimNodeCore/IUAFTransitionNode.h"
 #include "UAF/AnimNodeCore/IUAFTransitionContainerNode.h"
@@ -41,60 +42,20 @@ namespace UE::UAF
 	 *
 	 * Instance data for a simple blend transition
 	 */
-	class FUAFSimpleTransition : public FUAFAnimNode, public IUAFTransitionNode, public IUAFTransitionContainerNode 
+	class FUAFSimpleTransition : public FUAFTimedTransition
 	{
 	public:
 		UE_API FUAFSimpleTransition(FUAFAnimGraphUpdateContext& Context, FUAFAnimNodePtr InSource, FUAFAnimNodePtr InTarget, float InDuration, EAlphaBlendOption InBlendOption);
 
-		// ITransitionInstance impl
-		virtual const FUAFAnimNodePtr& GetSource() const override;
-		virtual const FUAFAnimNodePtr& GetTarget() const override;
-		UE_API virtual FUAFAnimNodePtr ReleaseSource() override;
-		UE_API virtual FUAFAnimNodePtr ReleaseTarget() override;
-		virtual bool IsComplete() const override;
-
-		// ITransitionContainer impl
-		UE_API virtual void NotifyTransitionComplete(const IUAFTransitionNode& TransitionNode) override;
-
 		// FUAFAnimNode impl
 		UE_API virtual void PreUpdate(FUAFAnimGraphUpdateContext& GraphContext) override;
-		UE_API virtual void* GetInterface(FUAFAnimNodeInterfaceId Id) override;
 		
 #if UAF_TRACE_ENABLED
 		UE_API virtual FString GetDebugName() const override;
 		UE_API virtual UStruct* GetDebugStruct() const override;
 #endif
-
-	protected:
-		static constexpr int32 SourceChildIndex = 0;
-		static constexpr int32 TargetChildIndex = 1;
-
-		float Duration = 0.1f;
-		EAlphaBlendOption BlendOption = UE::Anim::DefaultBlendOption;
-
-		float TimeRemaining = 0.0f;
-		bool bNestedTransitionComplete = false;
-
 		FUAFBlendTwoAnimOp BlendAnimOp;
 	};
-
-	//////////////////////////////////////////////////////////////////////////
-	// Inline implementation
-
-	inline const FUAFAnimNodePtr& FUAFSimpleTransition::GetSource() const
-	{
-		return GetChildAt(SourceChildIndex);
-	}
-
-	inline const FUAFAnimNodePtr& FUAFSimpleTransition::GetTarget() const
-	{
-		return GetChildAt(TargetChildIndex);
-	}
-
-	inline bool FUAFSimpleTransition::IsComplete() const
-	{
-		return TimeRemaining <= 0.0f;
-	}
 }
 
 #undef UE_API
diff --git a/Engine/Plugins/Experimental/UAF/UAFAnimNode/Source/UAFAnimNode/Public/UAF/AnimNodes/UAFTimedTransition.h b/Engine/Plugins/Experimental/UAF/UAFAnimNode/Source/UAFAnimNode/Public/UAF/AnimNodes/UAFTimedTransition.h
new file mode 100644
--- /dev/null
+++ b/Engine/Plugins/Experimental/UAF/UAFAnimNode/Source/UAFAnimNode/Public/UAF/AnimNodes/UAFTimedTransition.h
@@ -0,0 +1,72 @@
+// Copyright Epic Games, Inc. All Rights Reserved.
+
+#pragma once
+
+#include "CoreMinimal.h"
+#include "AlphaBlend.h"
+#include "UAF/AnimNodeCore/UAFAnimNode.h"
+#include "UAF/AnimNodeCore/IUAFTransitionNode.h"
+#include "UAF/AnimNodeCore/IUAFTransitionContainerNode.h"
+
+#define UE_API UAFANIMNODE_API
+
+namespace UE::UAF
+{
+	/**
+	 * FUAFTimedTransition
+	 */
+	class FUAFTimedTransition : public FUAFAnimNode, public IUAFTransitionNode, public IUAFTransitionContainerNode 
+	{
+	public:
+		UE_API FUAFTimedTransition(FUAFAnimGraphUpdateContext& Context, FUAFAnimNodePtr InSource, FUAFAnimNodePtr InTarget, float InDuration, EAlphaBlendOption InBlendOption);
+
+		// ITransitionInstance impl
+		virtual const FUAFAnimNodePtr& GetSource() const override;
+		virtual const FUAFAnimNodePtr& GetTarget() const override;
+		UE_API virtual FUAFAnimNodePtr ReleaseSource() override;
+		UE_API virtual FUAFAnimNodePtr ReleaseTarget() override;
+		virtual bool IsComplete() const override;
+
+		// ITransitionContainer impl
+		UE_API virtual void NotifyTransitionComplete(const IUAFTransitionNode& TransitionNode) override;
+
+		// FUAFAnimNode impl
+		UE_API virtual void PreUpdate(FUAFAnimGraphUpdateContext& GraphContext) override;
+		UE_API virtual void* GetInterface(FUAFAnimNodeInterfaceId Id) override;
+		
+#if UAF_TRACE_ENABLED
+		UE_API virtual FString GetDebugName() const override;
+		UE_API virtual UStruct* GetDebugStruct() const override;
+#endif
+
+	protected:
+		static constexpr int32 SourceChildIndex = 0;
+		static constexpr int32 TargetChildIndex = 1;
+
+		float Duration = 0.1f;
+		EAlphaBlendOption BlendOption = UE::Anim::DefaultBlendOption;
+
+		float TimeRemaining = 0.0f;
+		bool bNestedTransitionComplete = false;
+	};
+	
+	//////////////////////////////////////////////////////////////////////////
+	// Inline implementation
+
+	inline const FUAFAnimNodePtr& FUAFTimedTransition::GetSource() const
+	{
+		return GetChildAt(SourceChildIndex);
+	}
+
+	inline const FUAFAnimNodePtr& FUAFTimedTransition::GetTarget() const
+	{
+		return GetChildAt(TargetChildIndex);
+	}
+
+	inline bool FUAFTimedTransition::IsComplete() const
+	{
+		return TimeRemaining <= 0.0f;
+	}
+}
+
+#undef UE_API
diff --git a/Engine/Plugins/Experimental/UAF/UAFAnimNode/Source/UAFAnimNode/Public/UAF/AnimOps/UAFBlendTwoAnimOp.h b/Engine/Plugins/Experimental/UAF/UAFAnimNode/Source/UAFAnimNode/Public/UAF/AnimOps/UAFBlendTwoAnimOp.h
--- a/Engine/Plugins/Experimental/UAF/UAFAnimNode/Source/UAFAnimNode/Public/UAF/AnimOps/UAFBlendTwoAnimOp.h
+++ b/Engine/Plugins/Experimental/UAF/UAFAnimNode/Source/UAFAnimNode/Public/UAF/AnimOps/UAFBlendTwoAnimOp.h
@@ -3,6 +3,7 @@
 #pragma once
 
 #include "CoreMinimal.h"
+#include "AlphaBlend.h"
 #include "UAF/AnimOpCore/UAFAnimOp.h"
 
 #include "UAFBlendTwoAnimOp.generated.h"
@@ -39,6 +40,8 @@ namespace UE::UAF
 		UE_API void SetInterpolationAlpha(float InterpolationAlpha);
 		UE_API void SetInterpolationCurve(const FName& AlphaSourceCurveName, const int8 AlphaCurveInputIndex, TFunction<float(float)> InputScaleBiasClampFn);
 
+		UE_API void SetBlendOption(EAlphaBlendOption InBlendOption);
+		
 		// FUAFAnimOp impl
 		UE_API virtual void EvaluateValues(FUAFAnimOpValueEvaluator& Evaluator) override;
 
@@ -62,6 +65,9 @@ namespace UE::UAF
 		UPROPERTY(VisibleAnywhere, Category = Properties)
 		int8 AlphaCurveInputIndex = INDEX_NONE;
 
+		UPROPERTY(VisibleAnywhere, Category = Properties)
+		EAlphaBlendOption BlendOption = UE::Anim::DefaultBlendOption;
+		
 		TFunction<float(float)> InputScaleBiasClampFn;
 	};
 }
diff --git a/Engine/Plugins/Experimental/UAF/UAFAnimNode/Source/UAFAnimNode/Public/UAF/AnimOps/UAFBlendTwoPerValueAnimOp.h b/Engine/Plugins/Experimental/UAF/UAFAnimNode/Source/UAFAnimNode/Public/UAF/AnimOps/UAFBlendTwoPerValueAnimOp.h
new file mode 100644
--- /dev/null
+++ b/Engine/Plugins/Experimental/UAF/UAFAnimNode/Source/UAFAnimNode/Public/UAF/AnimOps/UAFBlendTwoPerValueAnimOp.h
@@ -0,0 +1,49 @@
+// Copyright Epic Games, Inc. All Rights Reserved.
+
+#pragma once
+
+#include "AlphaBlend.h"
+#include "CoreMinimal.h"
+#include "UAF/AnimOpCore/UAFAnimOp.h"
+
+#include "UAFBlendTwoPerValueAnimOp.generated.h"
+
+#define UE_API UAFANIMNODE_API
+
+class UUAFBlendProfile;
+
+namespace UE::UAF
+{
+	class FValueBundle;
+}
+
+namespace UE::UAF
+{
+	USTRUCT()
+	struct FUAFBlendTwoPerValueAnimOp : public FUAFAnimOp
+	{
+		GENERATED_BODY()
+		UAF_DECLARE_ANIMOP(FUAFBlendTwoPerValueAnimOp)
+
+		UE_API FUAFBlendTwoPerValueAnimOp();
+		UE_API FUAFBlendTwoPerValueAnimOp(TNonNullPtr<UUAFBlendProfile> BlendProfile, EAlphaBlendOption BlendOption);
+
+		UE_API void SetLinearWeightB(float InterpolationAlpha);
+
+		// FUAFAnimOp impl
+		UE_API virtual void EvaluateValues(FUAFAnimOpValueEvaluator& Evaluator) override;
+
+	protected:
+		// The interpolation alpha between the two input keyframes
+		UPROPERTY(VisibleAnywhere, Category = Properties)
+		float LinearGlobalWeightB = 0.0f;
+
+		UPROPERTY(VisibleAnywhere, Category = Properties)
+		TObjectPtr<UUAFBlendProfile> BlendProfile = nullptr;
+
+		UPROPERTY(VisibleAnywhere, Category = Properties)
+		EAlphaBlendOption BlendOption = UE::Anim::DefaultBlendOption;
+	};
+}
+
+#undef UE_API
\ No newline at end of file
