Fix edge case where a nested transition ends simultaneously with its parent transition and so never finishes.

#Synced-CL 52174097

diff --git a/Engine/Plugins/Experimental/UAF/UAFAnimNode/Source/UAFAnimNode/Private/UAF/AnimNodes/UAFBlendStack.cpp b/Engine/Plugins/Experimental/UAF/UAFAnimNode/Source/UAFAnimNode/Private/UAF/AnimNodes/UAFBlendStack.cpp
--- a/Engine/Plugins/Experimental/UAF/UAFAnimNode/Source/UAFAnimNode/Private/UAF/AnimNodes/UAFBlendStack.cpp
+++ b/Engine/Plugins/Experimental/UAF/UAFAnimNode/Source/UAFAnimNode/Private/UAF/AnimNodes/UAFBlendStack.cpp
@@ -38,7 +38,7 @@ void FUAFBlendStack::PruneTransitions()
 				if (TransitionInstance->IsComplete())
 				{
 					// Our transition is complete, prune it
-					FUAFAnimNodePtr Target = TransitionInstance->ReleaseTarget();
+					FUAFAnimNodePtr Target = TransitionInstance->Prune();
 					SetChildAt(0, Target);
 				}
 			}
diff --git a/Engine/Plugins/Experimental/UAF/UAFAnimNode/Source/UAFAnimNode/Private/UAF/AnimNodes/UAFTimedTransition.cpp b/Engine/Plugins/Experimental/UAF/UAFAnimNode/Source/UAFAnimNode/Private/UAF/AnimNodes/UAFTimedTransition.cpp
--- a/Engine/Plugins/Experimental/UAF/UAFAnimNode/Source/UAFAnimNode/Private/UAF/AnimNodes/UAFTimedTransition.cpp
+++ b/Engine/Plugins/Experimental/UAF/UAFAnimNode/Source/UAFAnimNode/Private/UAF/AnimNodes/UAFTimedTransition.cpp
@@ -112,7 +112,7 @@ namespace UE::UAF
 						{
 							if (NestedTransition->IsComplete())
 							{
-								FUAFAnimNodePtr NewSource = NestedTransition->ReleaseTarget();
+								FUAFAnimNodePtr NewSource = NestedTransition->Prune();
 								SetChildAt(SourceChildIndex, NewSource);
 							}
 						}
@@ -128,7 +128,23 @@ namespace UE::UAF
 				GetChildAt(SourceChildIndex)->SetTotalWeight(CurrentWeight * SourceWeight);
 				GetChildAt(TargetChildIndex)->SetTotalWeight(CurrentWeight * TargetWeight);
 			}
-		}	
+		}
+		else
+		{
+			// verify that the post anim op was cleared when we became complete
+			ensure(GetPostAnimOp() == nullptr);
+			
+			if (FUAFAnimNode* ParentNode = GetParent())
+			{
+				// if we are here, we've already notified our parent that we can be pruned, but have not been pruned.
+				// notify again, as this can happen if this node was re-parented before it's parent had a chance to prune it.
+                			
+				if (IUAFTransitionContainerNode* ParentContainer = ParentNode->GetInterface<IUAFTransitionContainerNode>())
+				{
+					ParentContainer->NotifyTransitionComplete(*this);
+				}
+			}
+		}
 	}
 	
 	void* FUAFTimedTransition::GetInterface(FUAFAnimNodeInterfaceId Id)
diff --git a/Engine/Plugins/Experimental/UAF/UAFAnimNode/Source/UAFAnimNode/Public/UAF/AnimNodeCore/IUAFTransitionNode.h b/Engine/Plugins/Experimental/UAF/UAFAnimNode/Source/UAFAnimNode/Public/UAF/AnimNodeCore/IUAFTransitionNode.h
--- a/Engine/Plugins/Experimental/UAF/UAFAnimNode/Source/UAFAnimNode/Public/UAF/AnimNodeCore/IUAFTransitionNode.h
+++ b/Engine/Plugins/Experimental/UAF/UAFAnimNode/Source/UAFAnimNode/Public/UAF/AnimNodeCore/IUAFTransitionNode.h
@@ -37,6 +37,32 @@ namespace UE::UAF
 
 		// Returns whether or not the transition has completed
 		[[nodiscard]] virtual bool IsComplete() const = 0;
+		
+		// Call Prune on a completed transition to get the unparented Target node
+		// Prune will also check if current target is another completed IUAFTransition,
+		// and will prune all the way down to a target node which is either not a transition, or not complete
+		[[nodiscard]] FUAFAnimNodePtr Prune()
+		{
+			ensure(IsComplete());
+			
+			FUAFAnimNodePtr Target = ReleaseTarget();
+
+			while(Target)
+			{
+				// If we contain any finished transtions prune all of them
+				if (IUAFTransitionNode* TransitionInstance = Target->GetInterface<IUAFTransitionNode>())
+				{
+					if (TransitionInstance->IsComplete())
+					{
+						Target = TransitionInstance->ReleaseTarget();
+						continue;
+					}
+				}
+				break;
+			} 
+
+			return Target;
+		}
 	};
 }
 
