Move data interface and provider data access inside read scope locks to ensure thread-safe access
Early-out on unchanged scrub time before acquiring provider locks

#Synced-CL 52253672

diff --git a/Engine/Plugins/Experimental/UAF/UAFAnimNode/Source/UAFAnimNodeEditor/Private/RewindDebugger/UAFAnimOpTrack.cpp b/Engine/Plugins/Experimental/UAF/UAFAnimNode/Source/UAFAnimNodeEditor/Private/RewindDebugger/UAFAnimOpTrack.cpp
--- a/Engine/Plugins/Experimental/UAF/UAFAnimNode/Source/UAFAnimNodeEditor/Private/RewindDebugger/UAFAnimOpTrack.cpp
+++ b/Engine/Plugins/Experimental/UAF/UAFAnimNode/Source/UAFAnimNodeEditor/Private/RewindDebugger/UAFAnimOpTrack.cpp
@@ -93,17 +93,17 @@ UAnimOpListDetailsObject* FUAFAnimOpTrack::InitializeDetailsObject()
 bool FUAFAnimOpTrack::UpdateInternal()
 {
 	TRACE_CPUPROFILER_EVENT_SCOPE(FUAFAnimOpTrack::UpdateInternal);
- 	
-	IRewindDebugger* RewindDebugger = IRewindDebugger::Instance();
- 	const TRange<double> ViewRange = RewindDebugger->GetCurrentViewRange();
-	
+
+	const IRewindDebugger* RewindDebugger = IRewindDebugger::Instance();
+	const TRange<double> ViewRange = RewindDebugger->GetCurrentViewRange();
+
 	const TraceServices::IAnalysisSession* AnalysisSession = RewindDebugger->GetAnalysisSession();
 
- 	bool bChanged = false;
-	
+	constexpr bool bChanged = false;
+
 	if (const FUAFAnimNodeProvider* UAFAnimNodeProvider = AnalysisSession->ReadProvider<FUAFAnimNodeProvider>("UAFAnimNodeProvider"))
 	{
-		double CurrentScrubTime = IRewindDebugger::Instance()->CurrentTraceTime();
+		const double CurrentScrubTime = IRewindDebugger::Instance()->CurrentTraceTime();
 
 		UAnimOpListDetailsObject* DetailsObject = DetailsObjectWeakPtr.Get();
 		if (DetailsObject == nullptr)
@@ -114,49 +114,50 @@ bool FUAFAnimOpTrack::UpdateInternal()
 
 		if (InstanceId == 0)
 		{
+			TraceServices::FProviderReadScopeLock UAFProviderReadScope(*UAFAnimNodeProvider);
 			UAFAnimNodeProvider->EnumerateEvaluationGraphs(ObjectId, [this](uint64 GraphId)
-		 	{
-		 		InstanceId = GraphId;
-		 	});
+				{
+					InstanceId = GraphId;
+				});
 		}
-		
-		if(InstanceId != 0)
+
+		if (InstanceId != 0
+			&& PreviousScrubTime != CurrentScrubTime)
 		{
-			if (const FAnimOpData* Data = UAFAnimNodeProvider->GetAnimOpData(InstanceId))
+			PreviousScrubTime = CurrentScrubTime;
+
+			bool bFrameFound;
+			TraceServices::FFrame MarkerFrame;
+			{
+				TraceServices::FAnalysisSessionReadScope SessionReadScope(*AnalysisSession);
+				const TraceServices::IFrameProvider& FrameProvider = TraceServices::ReadFrameProvider(*AnalysisSession);
+				bFrameFound = FrameProvider.GetFrameFromTime(TraceFrameType_Game, CurrentScrubTime, MarkerFrame);
+			}
+
+			if (bFrameFound)
 			{
-				if (PreviousScrubTime != CurrentScrubTime)
+				TraceServices::FProviderReadScopeLock UAFProviderReadScope(*UAFAnimNodeProvider);
+				if (const FAnimOpData* Data = UAFAnimNodeProvider->GetAnimOpData(InstanceId))
 				{
-					PreviousScrubTime = CurrentScrubTime;
-
-					bool bFrameFound;
-					TraceServices::FFrame MarkerFrame;
-					{
-						TraceServices::FAnalysisSessionReadScope SessionReadScope(*AnalysisSession);
-						const TraceServices::IFrameProvider& FrameProvider = TraceServices::ReadFrameProvider(*AnalysisSession);
-						bFrameFound = FrameProvider.GetFrameFromTime(TraceFrameType_Game, CurrentScrubTime, MarkerFrame);
-					}
-		
-					if (bFrameFound)
-					{
-						const IGameplayProvider* GameplayProvider = AnalysisSession->ReadProvider<IGameplayProvider>("GameplayProvider");
-						TraceServices::FProviderReadScopeLock GameplayProviderReadScope(*GameplayProvider);
-						TraceServices::FProviderReadScopeLock UAFAnimNodeProviderReadScope(*UAFAnimNodeProvider);
-						Data->AnimOpTimeline.EnumerateEvents(MarkerFrame.StartTime, MarkerFrame.EndTime, [DetailsObject, GameplayProvider](double InStartTime, double InEndTime, uint32 InDepth, const TArray<uint8>& VariableData)
-						{				
+					const IGameplayProvider* GameplayProvider = AnalysisSession->ReadProvider<IGameplayProvider>("GameplayProvider");
+					TraceServices::FProviderReadScopeLock GameplayProviderReadScope(*GameplayProvider);
+
+					Data->AnimOpTimeline.EnumerateEvents(MarkerFrame.StartTime, MarkerFrame.EndTime,
+						[DetailsObject, GameplayProvider](double InStartTime, double InEndTime, uint32 InDepth, const TArray<uint8>& VariableData)
+						{
 							FMemoryReader Reader(VariableData);
 							FObjectAsTraceIdProxyArchiveReader Archive(Reader, GameplayProvider);
 
 							static const FUAFInstancedAnimOpList Defaults;
 							FUAFInstancedAnimOpList::StaticStruct()->SerializeItem(Archive, &DetailsObject->AnimOps, &Defaults);
-								
+
 							return TraceServices::EEventEnumerate::Stop;
 						});
-					}
 				}
 			}
 
 			// TArray<uint64, TInlineAllocator<32>> CurrentChildren;
-		
+
 			// // update/create child tracks
 			// AnimNextProvider->EnumerateChildInstances(InstanceId, [this, &bChanged, &CurrentChildren, &ViewRange](const FDataInterfaceData& ChildData)
 			// {
@@ -183,8 +184,8 @@ bool FUAFAnimOpTrack::UpdateInternal()
 		// 	bChanged |= Child->Update();
 		// }
 	}
-	
- 	return bChanged;
+
+	return bChanged;
 }
 
 bool FUAFAnimOpTrackCreator::HasDebugInfoInternal(const RewindDebugger::FObjectId& InObjectId) const
diff --git a/Engine/Plugins/Experimental/UAF/UAFAnimNode/Source/UAFAnimNodeEditor/Private/RewindDebugger/UAFSequenceInfoTrack.cpp b/Engine/Plugins/Experimental/UAF/UAFAnimNode/Source/UAFAnimNodeEditor/Private/RewindDebugger/UAFSequenceInfoTrack.cpp
--- a/Engine/Plugins/Experimental/UAF/UAFAnimNode/Source/UAFAnimNodeEditor/Private/RewindDebugger/UAFSequenceInfoTrack.cpp
+++ b/Engine/Plugins/Experimental/UAF/UAFAnimNode/Source/UAFAnimNodeEditor/Private/RewindDebugger/UAFSequenceInfoTrack.cpp
@@ -280,16 +280,14 @@ bool FUAFSequenceInfoTrack::UpdateInternal()
 {
 	TRACE_CPUPROFILER_EVENT_SCOPE(FUAFSequenceInfoTrack::UpdateInternal);
 
-	IRewindDebugger* RewindDebugger = IRewindDebugger::Instance();
-
+	const IRewindDebugger* RewindDebugger = IRewindDebugger::Instance();
 	const TraceServices::IAnalysisSession* AnalysisSession = RewindDebugger->GetAnalysisSession();
 
-	bool bChanged = false;
+	constexpr bool bChanged = false;
 
-	if (const FUAFAnimNodeProvider* UAFAnimNodeProvider = IRewindDebugger::Instance()->GetAnalysisSession()->ReadProvider<
-		FUAFAnimNodeProvider>("UAFAnimNodeProvider"))
+	if (const FUAFAnimNodeProvider* UAFAnimNodeProvider = AnalysisSession->ReadProvider<FUAFAnimNodeProvider>("UAFAnimNodeProvider"))
 	{
-		double CurrentScrubTime = IRewindDebugger::Instance()->CurrentTraceTime();
+		const double CurrentScrubTime = IRewindDebugger::Instance()->CurrentTraceTime();
 
 		UUAFSequenceInfoDetailsObject* DetailsObject = DetailsObjectWeakPtr.Get();
 		if (DetailsObject == nullptr)
@@ -300,35 +298,36 @@ bool FUAFSequenceInfoTrack::UpdateInternal()
 
 		if (InstanceId == 0)
 		{
+			TraceServices::FProviderReadScopeLock UAFProviderReadScope(*UAFAnimNodeProvider);
 			UAFAnimNodeProvider->EnumerateEvaluationGraphs(ObjectId, [this](uint64 GraphId)
-			{
-				InstanceId = GraphId;
-			});
+				{
+					InstanceId = GraphId;
+				});
 		}
 
-		if (InstanceId != 0)
+		if (InstanceId != 0
+			&& PreviousScrubTime != CurrentScrubTime)
 		{
-			if (const FAnimOpData* Data = UAFAnimNodeProvider->GetAnimOpData(InstanceId))
+			PreviousScrubTime = CurrentScrubTime;
+
+			bool bFrameFound;
+			TraceServices::FFrame MarkerFrame;
+			{
+				TraceServices::FAnalysisSessionReadScope SessionReadScope(*AnalysisSession);
+				const TraceServices::IFrameProvider& FrameProvider = TraceServices::ReadFrameProvider(*AnalysisSession);
+				bFrameFound = FrameProvider.GetFrameFromTime(TraceFrameType_Game, CurrentScrubTime, MarkerFrame);
+			}
+
+			if (bFrameFound)
 			{
-				if (PreviousScrubTime != CurrentScrubTime)
+				TraceServices::FProviderReadScopeLock UAFProviderReadScope(*UAFAnimNodeProvider);
+				if (const FAnimOpData* Data = UAFAnimNodeProvider->GetAnimOpData(InstanceId))
 				{
-					PreviousScrubTime = CurrentScrubTime;
-
-					bool bFrameFound;
-					TraceServices::FFrame MarkerFrame;
-					{
-						TraceServices::FAnalysisSessionReadScope SessionReadScope(*AnalysisSession);
-						const TraceServices::IFrameProvider& FrameProvider = TraceServices::ReadFrameProvider(*AnalysisSession);
-						bFrameFound = FrameProvider.GetFrameFromTime(TraceFrameType_Game, CurrentScrubTime, MarkerFrame);
-					}
-		
-					if (bFrameFound)
-					{
-						const IGameplayProvider* GameplayProvider = AnalysisSession->ReadProvider<IGameplayProvider>("GameplayProvider");
-						TraceServices::FProviderReadScopeLock GameplayProviderReadScope(*GameplayProvider);
-
-						Data->AnimOpTimeline.EnumerateEvents(MarkerFrame.StartTime, MarkerFrame.EndTime,
-						[DetailsObject, GameplayProvider, this] ( double InStartTime, double InEndTime, uint32 InDepth, const TArray<uint8>& VariableData)
+					const IGameplayProvider* GameplayProvider = AnalysisSession->ReadProvider<IGameplayProvider>("GameplayProvider");
+					TraceServices::FProviderReadScopeLock GameplayProviderReadScope(*GameplayProvider);
+
+					Data->AnimOpTimeline.EnumerateEvents(MarkerFrame.StartTime, MarkerFrame.EndTime,
+						[DetailsObject, GameplayProvider, this](double InStartTime, double InEndTime, uint32 InDepth, const TArray<uint8>& VariableData)
 						{
 							// Why do this if panel is not selected?
 							FMemoryReader Reader(VariableData);
@@ -342,7 +341,6 @@ bool FUAFSequenceInfoTrack::UpdateInternal()
 							DetailsView->ForceRefresh(); // We need to force refresh to get the sync markers to refresh. todo: how to handle without recreating UI?
 							return TraceServices::EEventEnumerate::Stop;
 						});
-					}
 				}
 			}
 		}
