Add UAFSteeringNode
Add IUAFRootMotionProvider to allow anim nodes to query root motion of other nodes

#Synced-CL 52271689

diff --git a/Engine/Plugins/Experimental/UAF/UAFAnimNode/Source/UAFAnimNode/Private/UAF/AnimNodes/UAFSequencePlayer.cpp b/Engine/Plugins/Experimental/UAF/UAFAnimNode/Source/UAFAnimNode/Private/UAF/AnimNodes/UAFSequencePlayer.cpp
--- a/Engine/Plugins/Experimental/UAF/UAFAnimNode/Source/UAFAnimNode/Private/UAF/AnimNodes/UAFSequencePlayer.cpp
+++ b/Engine/Plugins/Experimental/UAF/UAFAnimNode/Source/UAFAnimNode/Private/UAF/AnimNodes/UAFSequencePlayer.cpp
@@ -91,6 +91,11 @@ namespace UE::UAF
 			return static_cast<IUAFAnimNodeTimeline*>(this);
 		}
 
+		if (Id == IUAFRootMotionProvider::InterfaceId)
+		{
+			return static_cast<IUAFRootMotionProvider*>(this);
+		}
+
 		return nullptr;
 	}
 
@@ -130,6 +135,24 @@ namespace UE::UAF
 		return DecompressAnimOp.GetTimeAdvanceResult();
 	}
 
+	FTransform FUAFSequencePlayer::ExtractRootMotion(float StartTime, float DeltaTime, bool bAllowLooping) const
+	{
+		if (const UAnimSequence* AnimSequence = DecompressAnimOp.GetAnimSequence())
+		{
+			// We do not check for lifetimes, assume the sequence is alive during pose list execution.
+			check(AnimSequence->IsValidLowLevel());
+
+			const bool bIsLooping = IsLooping();
+
+			// Make sure we don't extract outside our valid range
+			StartTime = FMath::Clamp(StartTime, 0.0f, AnimSequence->GetPlayLength());
+
+			return AnimSequence->ExtractRootMotion(FAnimExtractContext(static_cast<double>(StartTime), true, FDeltaTimeRecord(DeltaTime), bAllowLooping && bIsLooping));
+		}
+
+		return FTransform::Identity;
+	}
+
 	bool FUAFSequencePlayer::IsLooping() const
 	{
 		return DecompressAnimOp.IsLooping();
diff --git a/Engine/Plugins/Experimental/UAF/UAFAnimNode/Source/UAFAnimNode/Private/UAF/AnimNodes/UAFSteeringNode.cpp b/Engine/Plugins/Experimental/UAF/UAFAnimNode/Source/UAFAnimNode/Private/UAF/AnimNodes/UAFSteeringNode.cpp
new file mode 100644
--- /dev/null
+++ b/Engine/Plugins/Experimental/UAF/UAFAnimNode/Source/UAFAnimNode/Private/UAF/AnimNodes/UAFSteeringNode.cpp
@@ -0,0 +1,81 @@
+// Copyright Epic Games, Inc. All Rights Reserved.
+
+#include "UAF/AnimNodes/UAFSteeringNode.h"
+
+#include "UAF/AnimOps/UAFNullAnimOp.h"
+#include "UAF/AnimNodes/IUAFRootMotionProvider.h"
+#include "UAF/AnimNodes/IUAFAnimNodeTimeline.h"
+
+namespace UE::UAF
+{
+FUAFAnimNodePtr FUAFSteeringNodeData::CreateInstance(FUAFAnimGraphUpdateContext& Context) const
+{
+	return MakeAnimNode<FUAFSteeringNode>(Context, *this);
+}
+
+FUAFSteeringNode::FUAFSteeringNode(FUAFAnimGraphUpdateContext& Context, const FUAFSteeringNodeData& InData)
+	: FUAFAnimNode(Context), Data(&InData)
+{
+	InitializeAs<FUAFSteeringNode>(Context);
+	
+	if (const FUAFAnimNodeData* ChildData = InData.Child.GetPtr())
+	{
+		if (FUAFAnimNodePtr Child = ChildData->CreateInstance(Context))
+		{
+			// Update interfaces
+			SteeringAnimOp.RootMotionProvider = Child->GetInterface<IUAFRootMotionProvider>();
+			SteeringAnimOp.Timeline = Child->GetInterface<IUAFAnimNodeTimeline>();
+
+			AddChild(MoveTemp(Child));
+		}
+	}
+
+	SteeringAnimOp.HostObject = Context.GetHostObject();
+	SetPostAnimOp(&SteeringAnimOp);
+
+}
+
+#if UAF_TRACE_ENABLED
+FString FUAFSteeringNode::GetDebugName() const
+{
+	static FString SteeringNodeName("Steering");
+	return SteeringNodeName;
+}
+
+UStruct* FUAFSteeringNode::GetDebugStruct() const
+{
+	return FUAFSteeringNodeData::StaticStruct();
+}
+
+void FUAFSteeringNode::PreUpdate(FUAFAnimGraphUpdateContext& GraphContext)
+{
+	// Update Op with bindable values
+	SteeringAnimOp.Alpha = Data->Alpha.GetValue(GraphContext.GetVariablesOwner());
+	SteeringAnimOp.AnimatedTargetTime = Data->AnimatedTargetTime.GetValue(GraphContext.GetVariablesOwner());
+	SteeringAnimOp.DisableAdditiveBelowSpeed = Data->DisableAdditiveBelowSpeed.GetValue(GraphContext.GetVariablesOwner());
+	SteeringAnimOp.DisableSteeringBelowSpeed = Data->DisableSteeringBelowSpeed.GetValue(GraphContext.GetVariablesOwner());
+	SteeringAnimOp.TargetOrientation = Data->TargetOrientation.GetValue(GraphContext.GetVariablesOwner());
+	SteeringAnimOp.RootBoneTransform = Data->RootBoneTransform.GetValue(GraphContext.GetVariablesOwner());
+	SteeringAnimOp.ProceduralTargetTime = Data->ProceduralTargetTime.GetValue(GraphContext.GetVariablesOwner());
+	SteeringAnimOp.RootMotionAngleThresholdDegrees = Data->RootMotionAngleThresholdDegrees.GetValue(GraphContext.GetVariablesOwner());
+	SteeringAnimOp.MinScaleRatio = Data->MinScaleRatio.GetValue(GraphContext.GetVariablesOwner());
+	SteeringAnimOp.MaxScaleRatio = Data->MaxScaleRatio.GetValue(GraphContext.GetVariablesOwner());
+	
+	FUAFAnimNode::PreUpdate(GraphContext);
+}
+#endif // UAF_TRACE_ENABLED
+
+void FUAFSteeringNode::PostUpdate(FUAFAnimGraphUpdateContext& GraphContext)
+{
+	FUAFAnimNode::PostUpdate(GraphContext);
+
+	SteeringAnimOp.DeltaTime = GraphContext.GetDeltaTime();
+
+	if (FUAFAnimNodePtr Child = GetChildAt(0))
+	{
+		// Update interfaces: they can change in PreUpdate for e.g. transitions
+		SteeringAnimOp.RootMotionProvider = Child->GetInterface<IUAFRootMotionProvider>();
+		SteeringAnimOp.Timeline = Child->GetInterface<IUAFAnimNodeTimeline>();
+	}
+}
+}
diff --git a/Engine/Plugins/Experimental/UAF/UAFAnimNode/Source/UAFAnimNode/Private/UAF/AnimOps/UAFSteeringAnimOp.cpp b/Engine/Plugins/Experimental/UAF/UAFAnimNode/Source/UAFAnimNode/Private/UAF/AnimOps/UAFSteeringAnimOp.cpp
new file mode 100644
--- /dev/null
+++ b/Engine/Plugins/Experimental/UAF/UAFAnimNode/Source/UAFAnimNode/Private/UAF/AnimOps/UAFSteeringAnimOp.cpp
@@ -0,0 +1,178 @@
+// Copyright Epic Games, Inc. All Rights Reserved.
+
+#include "UAF/AnimOps/UAFSteeringAnimOp.h"
+
+#include "Animation/AnimRootMotionProvider.h"
+#include "Math/SpringMath.h"
+#include "HAL/IConsoleManager.h"
+#include "UAF/AnimNodes/IUAFAnimNodeTimeline.h"
+#include "UAF/AnimNodes/IUAFRootMotionProvider.h"
+#include "UAF/AnimOpCore/UAFAnimOpValueEvaluator.h"
+#include "VisualLogger/VisualLogger.h"
+
+bool gAnimNodeSteeringEnabled = true;
+static FAutoConsoleVariableRef CVarUAFSteeringEnabled(
+	TEXT("a.UAF.Steering.Enabled"),
+	gAnimNodeSteeringEnabled,
+	TEXT("True will enable steering for UAF. Equivalent to setting alpha to non-zero.")
+);
+
+namespace UE::UAF
+{
+
+FUAFSteeringAnimOp::FUAFSteeringAnimOp() : FUAFAnimOp(1)
+{
+	InitializeAs<FUAFSteeringAnimOp>();
+}
+
+void FUAFSteeringAnimOp::EvaluateValues(FUAFAnimOpValueEvaluator& Evaluator)
+{
+	if (DeltaTime <= UE_SMALL_NUMBER)
+	{
+		return;
+	}
+
+	if (Alpha <= UE_SMALL_NUMBER)
+	{
+		return;
+	}
+
+	if (gAnimNodeSteeringEnabled == false)
+	{
+		return;
+	}
+
+	// Get the root motion attribute from the evaluator
+	FTransform ThisFrameRootMotionTransform = FTransform::Identity;
+	
+
+	// Feedback: Would not have been able to discover this API myself, had to read UAFDecompressAnimSequenceAnimOp
+	const FUAFAnimOpValueEvaluationContext& AnimOpEvaluationContext = Evaluator.GetActiveEvaluationContext();
+
+	FValueBundleCoWRef* ValueBundle = Evaluator.GetEvaluationStack().PeekMutable();
+	check (ValueBundle != nullptr)
+	
+	FAttributeNamedSetPtr NamedSet = ValueBundle->Get().GetNamedSet();
+	if (NamedSet.IsValid() == false)
+	{
+		return;
+	}
+	
+	FAttributeTypedSetPtr AttributeTypedSet = NamedSet->FindTypedSet<FTransformAnimationAttribute>();
+	if (AttributeTypedSet.IsValid() == false)
+	{
+		return;
+	}
+
+	// AttributeName: RootMotionDelta
+	const FAttributeSetIndex SetIndex = AttributeTypedSet->FindIndex(UE::Anim::IAnimRootMotionProvider::AttributeName);
+	if(SetIndex.IsValid() == false)
+	{
+		return;
+	}
+	
+	FAttributeMappingKey MappingKey = FAttributeMappingKey::MakeFromTo<FTransformAnimationAttribute>();
+	TBoundValueMap<FTransformAnimationAttribute>* Map = ValueBundle->GetMutable().GetBoundValueMaps().Find<FTransformAnimationAttribute>(MappingKey);
+	checkf(Map != nullptr, TEXT("This attribute exists within our named set but isn't present in the output value bundle"));
+
+	ThisFrameRootMotionTransform = (*Map)[SetIndex].Value;
+
+	float CurrentSpeed = ThisFrameRootMotionTransform.GetTranslation().Length() / DeltaTime;
+	if (CurrentSpeed < DisableSteeringBelowSpeed)
+	{
+		return;
+	}
+
+	FQuat RootBoneRotation = RootBoneTransform.GetRotation();
+#if ENABLE_ANIM_DEBUG
+	UE_VLOG_ARROW(HostObject, "Steering", Display,
+		RootBoneTransform.GetLocation(),
+		RootBoneTransform.GetLocation() + RootBoneRotation.GetRightVector() * 90,
+		FColor::Green, TEXT(""));
+
+	UE_VLOG_ARROW(HostObject, "Steering", Display,
+		RootBoneTransform.GetLocation(),
+		RootBoneTransform.GetLocation() + TargetOrientation.GetRightVector() * 100,
+		FColor::Blue, TEXT(""));
+#endif // ENABLE_ANIM_DEBUG
+
+	FQuat DeltaToTargetOrientation = RootBoneRotation.Inverse() * TargetOrientation;
+
+	if (Timeline != nullptr && RootMotionProvider != nullptr && AnimatedTargetTime > 0.0f)
+	{
+		// Root motion prediction
+		float CurrentTime = Timeline->GetCurrentTime();
+
+		FTransform PredictedRootMotionTransform = RootMotionProvider->ExtractRootMotion(CurrentTime, AnimatedTargetTime, Timeline->IsLooping());
+		FQuat PredictedRootMotionQuat = PredictedRootMotionTransform.GetRotation();
+		FRotator PredictedRootMotionRot = FRotator(PredictedRootMotionQuat);
+		float PredictedRootMotionYaw = PredictedRootMotionRot.Yaw;
+
+		if (fabs(PredictedRootMotionYaw) > RootMotionAngleThresholdDegrees)
+		{
+#if ENABLE_ANIM_DEBUG
+			UE_VLOG_ARROW(HostObject, "Steering", Display,
+				RootBoneTransform.GetLocation(),
+				RootBoneTransform.GetLocation() + (PredictedRootMotionQuat * RootBoneRotation).GetRightVector() * 100,
+				FColor::Orange, TEXT(""));
+#endif // ENABLE_ANIM_DEBUG 
+
+			float YawToTargetOrientation = FRotator(DeltaToTargetOrientation).Yaw;
+
+			// pick the rotation direction that is the shortest path from the endpoint of the current animated rotation
+			if (PredictedRootMotionYaw - YawToTargetOrientation > 180)
+			{
+				YawToTargetOrientation += 360;
+			}
+			else if (YawToTargetOrientation - PredictedRootMotionYaw > 180)
+			{
+				YawToTargetOrientation -= 360;
+			}
+
+			float Ratio = YawToTargetOrientation / PredictedRootMotionYaw;
+			Ratio = FMath::Clamp(Ratio, MinScaleRatio, MaxScaleRatio);
+
+			// Account for alpha
+			Ratio = FMath::Lerp(1.0f, Ratio, Alpha);
+
+			FRotator ThisFrameRootMotionRotator(ThisFrameRootMotionTransform.GetRotation());
+			ThisFrameRootMotionRotator.Yaw *= Ratio;
+			ThisFrameRootMotionTransform.SetRotation(FQuat(ThisFrameRootMotionRotator));
+
+			// Account for future scaling in linear error correction
+			PredictedRootMotionRot.Yaw *= Ratio;
+			PredictedRootMotionQuat = PredictedRootMotionRot.Quaternion();
+
+			DeltaToTargetOrientation = PredictedRootMotionQuat.Inverse() * RootBoneRotation.Inverse() * TargetOrientation;
+		}
+	}
+
+	if (CurrentSpeed > DisableAdditiveBelowSpeed)
+	{
+		// Apply linear correction
+		FQuat LinearCorrection = FQuat::Identity;
+		SpringMath::CriticalSpringDamperQuat(LinearCorrection, AngularVelocity, DeltaToTargetOrientation, SpringMath::HalfLifeToSmoothingTime(ProceduralTargetTime), DeltaTime);
+
+#if ENABLE_ANIM_DEBUG
+		UE_VLOG_ARROW(HostObject, "Steering", Display,
+			RootBoneTransform.GetLocation(),
+			RootBoneTransform.GetLocation() + (RootBoneTransform.GetRotation() * LinearCorrection).GetRightVector() * 120,
+			FColor::Magenta, TEXT(""));
+#endif
+
+		FQuat ThisFrameRotation = ThisFrameRootMotionTransform.GetRotation() * LinearCorrection;
+
+#if ENABLE_ANIM_DEBUG
+		UE_VLOG_ARROW(HostObject, "Steering", Display,
+			RootBoneTransform.GetLocation(),
+			RootBoneTransform.GetLocation() + (RootBoneTransform.GetRotation() * ThisFrameRotation).GetRightVector() * 140,
+			FColor::Red, TEXT(""));
+#endif
+
+		ThisFrameRootMotionTransform.SetRotation(FQuat::Slerp(ThisFrameRootMotionTransform.GetRotation(), ThisFrameRotation, Alpha));
+	}
+
+	// Set back the updated root motion transform
+	(*Map)[SetIndex].Value = ThisFrameRootMotionTransform;
+}
+}
diff --git a/Engine/Plugins/Experimental/UAF/UAFAnimNode/Source/UAFAnimNode/Public/UAF/AnimNodes/IUAFRootMotionProvider.h b/Engine/Plugins/Experimental/UAF/UAFAnimNode/Source/UAFAnimNode/Public/UAF/AnimNodes/IUAFRootMotionProvider.h
new file mode 100644
--- /dev/null
+++ b/Engine/Plugins/Experimental/UAF/UAFAnimNode/Source/UAFAnimNode/Public/UAF/AnimNodes/IUAFRootMotionProvider.h
@@ -0,0 +1,20 @@
+// Copyright Epic Games, Inc. All Rights Reserved.
+
+#pragma once
+
+
+#include "CoreMinimal.h"
+#include "UAF/AnimNodeCore/UAFAnimNodeInterfaceId.h"
+
+
+namespace UE::UAF
+{
+struct IUAFRootMotionProvider
+{
+	static constexpr FUAFAnimNodeInterfaceId InterfaceId = FUAFAnimNodeInterfaceId::MakeFromString(TEXT("IUAFRootMotionProvider"));
+
+	virtual ~IUAFRootMotionProvider() = default;
+
+	[[nodiscard]] virtual FTransform ExtractRootMotion(float StartTime, float DeltaTime, bool AllowLooping) const = 0; 
+};
+}
diff --git a/Engine/Plugins/Experimental/UAF/UAFAnimNode/Source/UAFAnimNode/Public/UAF/AnimNodes/UAFSequencePlayer.h b/Engine/Plugins/Experimental/UAF/UAFAnimNode/Source/UAFAnimNode/Public/UAF/AnimNodes/UAFSequencePlayer.h
--- a/Engine/Plugins/Experimental/UAF/UAFAnimNode/Source/UAFAnimNode/Public/UAF/AnimNodes/UAFSequencePlayer.h
+++ b/Engine/Plugins/Experimental/UAF/UAFAnimNode/Source/UAFAnimNode/Public/UAF/AnimNodes/UAFSequencePlayer.h
@@ -3,6 +3,7 @@
 #pragma once
 
 #include "CoreMinimal.h"
+#include "IUAFRootMotionProvider.h"
 #include "UAF/AnimNodeCore/UAFAnimNodeData.h"
 #include "UAF/AnimNodeCore/UAFAnimNode.h"
 #include "UAF/AnimNodes/IUAFAnimNodeDataHasAsset.h"
@@ -45,7 +46,7 @@ namespace UE::UAF
 	};
 
 	// Anim node instance for playing an AnimSequence 
-	struct FUAFSequencePlayer : public FUAFAnimNode, public IUAFAnimNodeTimeline
+	struct FUAFSequencePlayer : public FUAFAnimNode, public IUAFAnimNodeTimeline, public IUAFRootMotionProvider
 	{
 		UE_API explicit FUAFSequencePlayer(FUAFAnimGraphUpdateContext& Context, const FUAFSequencePlayerData* InData);
 		UE_API explicit FUAFSequencePlayer(FUAFAnimGraphUpdateContext& Context, const UAnimSequence* Sequence, EAnimAssetLoopMode LoopMode = EAnimAssetLoopMode::Auto);
@@ -67,6 +68,9 @@ namespace UE::UAF
 		UE_API virtual float GetLength() const override;
 		UE_API virtual ETypeAdvanceAnim GetTimeAdvanceResult() const override;
 
+		// IUAFRootMotionProvider
+		UE_API virtual FTransform ExtractRootMotion(float StartTime, float DeltaTime, bool AllowLooping) const override;
+
 	private:
 		FUAFDecompressAnimSequenceAnimOp DecompressAnimOp;
 	};
diff --git a/Engine/Plugins/Experimental/UAF/UAFAnimNode/Source/UAFAnimNode/Public/UAF/AnimNodes/UAFSteeringNode.h b/Engine/Plugins/Experimental/UAF/UAFAnimNode/Source/UAFAnimNode/Public/UAF/AnimNodes/UAFSteeringNode.h
new file mode 100644
--- /dev/null
+++ b/Engine/Plugins/Experimental/UAF/UAFAnimNode/Source/UAFAnimNode/Public/UAF/AnimNodes/UAFSteeringNode.h
@@ -0,0 +1,85 @@
+// Copyright Epic Games, Inc. All Rights Reserved.
+
+#pragma once
+
+#include "CoreMinimal.h"
+#include "BindableValue/UAFBindableTypes.h"
+#include "UAF/AnimNodeCore/UAFAnimNode.h"
+#include "UAF/AnimNodeCore/UAFAnimNodeData.h"
+#include "UAF/AnimOps/UAFSteeringAnimOp.h"
+#include "UAFSteeringNode.generated.h"
+
+#define UE_API UAFANIMNODE_API
+
+namespace UE::UAF
+{
+
+USTRUCT(DisplayName = "Steering")
+struct FUAFSteeringNodeData : public FUAFAnimNodeData
+{
+	GENERATED_BODY()
+
+	UPROPERTY(EditAnywhere, Category = "Data")
+	TInstancedStruct<FUAFAnimNodeData> Child;
+
+	UPROPERTY(EditAnywhere, Category = Settings)
+	FBindableFloat Alpha = 1.0f;
+
+	UPROPERTY(EditAnywhere, Category = Settings)
+	FBindableQuat TargetOrientation;
+
+	UPROPERTY(EditAnywhere, Category = Settings)
+	FBindableTransform RootBoneTransform;
+
+	// The number of seconds in the future before we should reach the TargetOrientation when play animations with no root motion rotation
+	UPROPERTY(EditAnywhere, Category=Settings)
+	FBindableFloat ProceduralTargetTime = 0.2f;
+
+	// The number of seconds in the future before we should reach the TargetOrientation when playing animations with root motion rotation
+	UPROPERTY(EditAnywhere, Category=Settings)
+	FBindableFloat AnimatedTargetTime = 0.2f;
+
+	UPROPERTY(EditAnywhere, Category=Settings)
+	FBindableFloat RootMotionAngleThresholdDegrees = 1.0f;
+
+	// Will clamp the scaling ratio applied to above this threashold. Any error remaining will be compensated for linearly (using ProceduralTargetTime)
+	UPROPERTY(EditAnywhere, Category = Settings)
+	FBindableFloat MinScaleRatio = 0.5f;
+
+	// Will clamp the scaling ratio applied to below this threashold. Any error remaining will be compensated for linearly (using ProceduralTargetTime)
+	UPROPERTY(EditAnywhere, Category = Settings)
+	FBindableFloat MaxScaleRatio = 1.5f;
+
+	// below this movement speed (based on the root motion in the animation) disable steering coming from the additive spring based correction
+	UPROPERTY(EditAnywhere, Category = Settings)
+	FBindableFloat DisableAdditiveBelowSpeed = -1.0f;
+
+	// below this movement speed (based on the root motion in the animation) disable steering completely (both scaling and additive)
+	UPROPERTY(EditAnywhere, Category=Settings)
+	FBindableFloat DisableSteeringBelowSpeed = 1.0f;
+
+	// FUAFTransitionNodeData impl
+	virtual FUAFAnimNodePtr CreateInstance(FUAFAnimGraphUpdateContext& Context) const override;
+};
+
+class FUAFSteeringNode : public FUAFAnimNode
+{
+public:
+	UE_API FUAFSteeringNode(FUAFAnimGraphUpdateContext& Context, const FUAFSteeringNodeData& InData);
+
+#if UAF_TRACE_ENABLED
+	virtual FString GetDebugName() const override;
+	virtual UStruct* GetDebugStruct() const override;
+#endif
+
+protected:
+	virtual void PreUpdate(FUAFAnimGraphUpdateContext& GraphContext) override;
+	virtual void PostUpdate(FUAFAnimGraphUpdateContext& GraphContext) override;
+
+	const FUAFSteeringNodeData* Data;
+
+	FUAFSteeringAnimOp SteeringAnimOp;
+};
+}
+
+#undef UE_API
\ No newline at end of file
diff --git a/Engine/Plugins/Experimental/UAF/UAFAnimNode/Source/UAFAnimNode/Public/UAF/AnimOps/UAFSteeringAnimOp.h b/Engine/Plugins/Experimental/UAF/UAFAnimNode/Source/UAFAnimNode/Public/UAF/AnimOps/UAFSteeringAnimOp.h
new file mode 100644
--- /dev/null
+++ b/Engine/Plugins/Experimental/UAF/UAFAnimNode/Source/UAFAnimNode/Public/UAF/AnimOps/UAFSteeringAnimOp.h
@@ -0,0 +1,79 @@
+// Copyright Epic Games, Inc. All Rights Reserved.
+
+#pragma once
+
+#include "CoreMinimal.h"
+#include "UAF/AnimOpCore/UAFAnimOp.h"
+
+#include "UAFSteeringAnimOp.generated.h"
+
+#define UE_API UAFANIMNODE_API
+
+namespace UE::UAF
+{
+struct IUAFAnimNodeTimeline;
+struct IUAFRootMotionProvider;
+class FValueBundle;
+}
+
+namespace UE::UAF
+{
+USTRUCT()
+struct FUAFSteeringAnimOp : public FUAFAnimOp
+{
+	GENERATED_BODY()
+	UAF_DECLARE_ANIMOP(FUAFSteeringAnimOp)
+
+	FUAFSteeringAnimOp();
+
+	// FUAFAnimOp impl
+	UE_API virtual void EvaluateValues(FUAFAnimOpValueEvaluator& Evaluator) override;
+
+	UPROPERTY()
+	float DeltaTime = 0.f;
+
+	UPROPERTY()
+	float Alpha = 1.0f;
+
+	UPROPERTY()
+	FVector AngularVelocity = FVector::ZeroVector;
+
+	UPROPERTY()
+	FQuat TargetOrientation;
+
+	UPROPERTY()
+	FTransform RootBoneTransform;
+
+	UPROPERTY()
+	float ProceduralTargetTime = 0.2f;
+
+	UPROPERTY()
+	float AnimatedTargetTime = 0.f;
+
+	UPROPERTY()
+	float DisableSteeringBelowSpeed = 0.f;
+
+	UPROPERTY()
+	float RootMotionAngleThresholdDegrees = 1.0f;
+
+	// below this movement speed (based on the root motion in the animation) disable steering coming from the additive spring based correction
+	UPROPERTY()
+	float DisableAdditiveBelowSpeed = -1.0f;
+
+	// Will clamp the scaling ratio applied to above this threashold. Any error remaining will be compensated for linearly (using ProceduralTargetTime)
+	UPROPERTY()
+	float MinScaleRatio = 0.5f;
+
+	// Will clamp the scaling ratio applied to below this threashold. Any error remaining will be compensated for linearly (using ProceduralTargetTime)
+	UPROPERTY()
+	float MaxScaleRatio = 1.5f;
+
+	UObject* HostObject = nullptr;
+	
+	// These interfaces are needed if AnimatedTargetTime is greater than 0 (i.e. if we want to take into account predicted root motion)
+	IUAFRootMotionProvider* RootMotionProvider = nullptr;
+	IUAFAnimNodeTimeline* Timeline = nullptr;
+};
+}
+
+#undef UE_API
\ No newline at end of file
