Add OffsetRootBone anim node
Fix steering node compile error in non trace builds

#Synced-CL 52271689

diff --git a/Engine/Plugins/Experimental/UAF/UAFAnimNode/Source/UAFAnimNode/Private/UAF/AnimNodes/UAFOffsetRootBoneNode.cpp b/Engine/Plugins/Experimental/UAF/UAFAnimNode/Source/UAFAnimNode/Private/UAF/AnimNodes/UAFOffsetRootBoneNode.cpp
new file mode 100644
--- /dev/null
+++ b/Engine/Plugins/Experimental/UAF/UAFAnimNode/Source/UAFAnimNode/Private/UAF/AnimNodes/UAFOffsetRootBoneNode.cpp
@@ -0,0 +1,71 @@
+// Copyright Epic Games, Inc. All Rights Reserved.
+
+#include "UAF/AnimNodes/UAFOffsetRootBoneNode.h"
+
+namespace UE::UAF
+{
+
+FUAFAnimNodePtr FUAFOffsetRootBoneNodeData::CreateInstance(FUAFAnimGraphUpdateContext& Context) const
+{
+	return MakeAnimNode<FUAFOffsetRootBoneNode>(Context, *this);
+}
+
+FUAFOffsetRootBoneNode::FUAFOffsetRootBoneNode(FUAFAnimGraphUpdateContext& Context, const FUAFOffsetRootBoneNodeData& InData)
+	: FUAFAnimNode(Context), Data(&InData)
+{
+	InitializeAs<FUAFOffsetRootBoneNode>(Context);
+
+	if (const FUAFAnimNodeData* ChildData = InData.Child.GetPtr())
+	{
+		if (FUAFAnimNodePtr Child = ChildData->CreateInstance(Context))
+		{
+			AddChild(MoveTemp(Child));
+		}
+	}
+
+	// Copy non-latent (constant) properties from data to AnimOp
+	OffsetRootBoneAnimOp.bClampToTranslationVelocity = InData.bClampToTranslationVelocity;
+	OffsetRootBoneAnimOp.bClampToRotationVelocity = InData.bClampToRotationVelocity;
+	OffsetRootBoneAnimOp.TranslationSpeedRatio = InData.TranslationSpeedRatio;
+	OffsetRootBoneAnimOp.RotationSpeedRatio = InData.RotationSpeedRatio;
+	OffsetRootBoneAnimOp.TeleportDistanceThreshold = InData.TeleportDistanceThreshold;
+	OffsetRootBoneAnimOp.bResetOnTeleport = InData.bResetOnTeleport;
+
+	OffsetRootBoneAnimOp.HostObject = Context.GetHostObject();
+	SetPostAnimOp(&OffsetRootBoneAnimOp);
+}
+
+#if UAF_TRACE_ENABLED
+FString FUAFOffsetRootBoneNode::GetDebugName() const
+{
+	static FString Name(TEXT("OffsetRootBone"));
+	return Name;
+}
+
+UStruct* FUAFOffsetRootBoneNode::GetDebugStruct() const
+{
+	return FUAFOffsetRootBoneNodeData::StaticStruct();
+}
+#endif // UAF_TRACE_ENABLED
+
+void FUAFOffsetRootBoneNode::PreUpdate(FUAFAnimGraphUpdateContext& GraphContext)
+{
+	OffsetRootBoneAnimOp.DeltaTime = GraphContext.GetDeltaTime();
+
+	// Resolve bindable values
+	FUAFAssetInstance* VarOwner = GraphContext.GetVariablesOwner();
+	OffsetRootBoneAnimOp.Alpha = Data->Alpha.GetValue(VarOwner);
+	OffsetRootBoneAnimOp.MeshComponentTransformWorld = Data->MeshComponentTransformWorld.GetValue(VarOwner);
+	OffsetRootBoneAnimOp.TranslationMode = Data->TranslationMode.GetValue<EUAFOffsetRootBoneNodeMode>(VarOwner);
+	OffsetRootBoneAnimOp.RotationMode = Data->RotationMode.GetValue<EUAFOffsetRootBoneNodeMode>(VarOwner);
+	OffsetRootBoneAnimOp.TranslationSmoothingTime = Data->TranslationSmoothingTime.GetValue(VarOwner);
+	OffsetRootBoneAnimOp.RotationSmoothingTime = Data->RotationSmoothingTime.GetValue(VarOwner);
+	OffsetRootBoneAnimOp.MaxTranslationError = Data->MaxTranslationError.GetValue(VarOwner);
+	OffsetRootBoneAnimOp.MaxRotationErrorDegrees = Data->MaxRotationErrorDegrees.GetValue(VarOwner);
+	OffsetRootBoneAnimOp.bOnGround = Data->bOnGround.GetValue(VarOwner);
+	OffsetRootBoneAnimOp.AnimatedGroundNormal = Data->AnimatedGroundNormal.GetValue(VarOwner);
+	
+	FUAFAnimNode::PreUpdate(GraphContext);
+}
+
+} // namespace UE::UAF
diff --git a/Engine/Plugins/Experimental/UAF/UAFAnimNode/Source/UAFAnimNode/Private/UAF/AnimOps/UAFOffsetRootBoneAnimOp.cpp b/Engine/Plugins/Experimental/UAF/UAFAnimNode/Source/UAFAnimNode/Private/UAF/AnimOps/UAFOffsetRootBoneAnimOp.cpp
new file mode 100644
--- /dev/null
+++ b/Engine/Plugins/Experimental/UAF/UAFAnimNode/Source/UAFAnimNode/Private/UAF/AnimOps/UAFOffsetRootBoneAnimOp.cpp
@@ -0,0 +1,351 @@
+// Copyright Epic Games, Inc. All Rights Reserved.
+
+#include "UAF/AnimOps/UAFOffsetRootBoneAnimOp.h"
+
+#include "Animation/AnimRootMotionProvider.h"
+#include "Math/SpringMath.h"
+#include "BoneIndices.h"
+#include "HAL/IConsoleManager.h"
+#include "UAF/AnimOpCore/UAFAnimOpValueEvaluator.h"
+#include "UAF/Attributes/AttributeBindingIndex.h"
+#include "UAF/Attributes/EngineAttributes.h"
+#include "VisualLogger/VisualLogger.h"
+
+bool gAnimNodeOffsetRootBoneEnabled = true;
+static FAutoConsoleVariableRef CVarUAFOffsetRootBoneEnabled(
+	TEXT("a.UAF.OffsetRootBone.Enabled"),
+	gAnimNodeOffsetRootBoneEnabled,
+	TEXT("True will enable Offset Root Bone for UAF. Equivalent to setting alpha to non-zero.")
+);
+
+namespace UE::UAF::OffsetRootBone
+{
+
+static bool ShouldExtractRootMotion(const EUAFOffsetRootBoneNodeMode OffsetMode)
+{
+	switch (OffsetMode)
+	{
+	case EUAFOffsetRootBoneNodeMode::Accumulate:
+	case EUAFOffsetRootBoneNodeMode::Interpolate:
+		return true;
+	case EUAFOffsetRootBoneNodeMode::Release:
+	default:
+		return false;
+	}
+}
+
+static bool ShouldCounterComponentDelta(const EUAFOffsetRootBoneNodeMode OffsetMode)
+{
+	switch (OffsetMode)
+	{
+	case EUAFOffsetRootBoneNodeMode::Accumulate:
+	case EUAFOffsetRootBoneNodeMode::Interpolate:
+		return false;
+	case EUAFOffsetRootBoneNodeMode::Release:
+	default:
+		return true;
+	}
+}
+
+} // namespace UE::UAF::OffsetRootBone
+
+namespace UE::UAF
+{
+
+FUAFOffsetRootBoneAnimOp::FUAFOffsetRootBoneAnimOp() : FUAFAnimOp(1)
+{
+	InitializeAs<FUAFOffsetRootBoneAnimOp>();
+}
+
+void FUAFOffsetRootBoneAnimOp::EvaluateValues(FUAFAnimOpValueEvaluator& Evaluator)
+{
+	QUICK_SCOPE_CYCLE_COUNTER(FUAFOffsetRootBoneAnimOp_EvaluateValues);
+
+	if (DeltaTime == 0.f)
+	{
+		return;
+	}
+
+	if (Alpha == 0.f)
+	{
+		bIsFirstUpdate = true;
+		return;
+	}
+
+	if (!gAnimNodeOffsetRootBoneEnabled)
+	{
+		bIsFirstUpdate = true;
+		return;
+	}
+
+	// Access the value bundle from the evaluation context
+	const FUAFAnimOpValueEvaluationContext& AnimOpEvaluationContext = Evaluator.GetActiveEvaluationContext();
+	FValueBundleCoWRef* ValueBundle = Evaluator.GetEvaluationStack().PeekMutable();
+	check (ValueBundle);
+	
+	FAttributeNamedSetPtr NamedSet = ValueBundle->Get().GetNamedSet();
+	if (!NamedSet.IsValid())
+	{
+		return;
+	}
+
+	// --- Read root bone transform ---
+	FAttributeTypedSetPtr BoneTypedSet = NamedSet->FindTypedSet<FBoneTransformAnimationAttribute>();
+	if (!BoneTypedSet.IsValid())
+	{
+		return;
+	}
+
+	const FAttributeBindingIndex RootBoneBindingIndex(FSkeletonPoseBoneIndex(0));
+	const FAttributeSetIndex RootBoneSetIndex = BoneTypedSet->GetIndex(RootBoneBindingIndex);
+	if (!RootBoneSetIndex.IsValid())
+	{
+		return;
+	}
+
+	FAttributeMappingKey BoneMappingKey = FAttributeMappingKey::MakeFromTo<FBoneTransformAnimationAttribute>();
+	TBoundValueMap<FBoneTransformAnimationAttribute>* BoneMap = ValueBundle->GetMutable().GetBoundValueMaps().Find<FBoneTransformAnimationAttribute>(BoneMappingKey);
+	if (!BoneMap)
+	{
+		return;
+	}
+
+	const FTransform InputBoneTransform = (*BoneMap)[RootBoneSetIndex].Value;
+
+	// --- Read root motion delta ---
+	FTransform RootMotionTransformDelta = FTransform::Identity;
+	bool bHasRootMotion = false;
+	if (FAttributeTypedSetPtr RootMotionTypedSet = NamedSet->FindTypedSet<FTransformAnimationAttribute>())
+	{
+		const FAttributeSetIndex RootMotionSetIndex = RootMotionTypedSet->FindIndex(UE::Anim::IAnimRootMotionProvider::AttributeName);
+		if (RootMotionSetIndex.IsValid())
+		{
+			FAttributeMappingKey RootMotionMappingKey = FAttributeMappingKey::MakeFromTo<FTransformAnimationAttribute>();
+			if (const TBoundValueMap<FTransformAnimationAttribute>* RootMotionMap = ValueBundle->Get().GetBoundValueMaps().Find<FTransformAnimationAttribute>(RootMotionMappingKey))
+			{
+				RootMotionTransformDelta = (*RootMotionMap)[RootMotionSetIndex].Value;
+				bHasRootMotion = true;
+			}
+		}
+	}
+
+	if (!bHasRootMotion)
+	{
+		return;
+	}
+
+	// --- Offset Root Bone calculation ---
+	// Ported from FAnimNextOffsetRootBoneTask::Execute
+
+	const FTransform ComponentTransform = MeshComponentTransformWorld;
+
+	if (bIsFirstUpdate)
+	{
+		SimulatedTranslation = ComponentTransform.GetTranslation();
+		SimulatedRotation = ComponentTransform.GetRotation();
+	}
+
+	// Note we set last to current on first update
+	FTransform LastComponentTransform = bIsFirstUpdate ? ComponentTransform : LastMeshComponentTransformWorld;
+	bIsFirstUpdate = false;
+
+	// Teleport detection: if the difference between the previous and current component transform
+	// is above the teleport threshold, adjust the simulated position to account for the teleport
+	if (FVector::Distance(LastComponentTransform.GetLocation(), ComponentTransform.GetLocation()) > TeleportDistanceThreshold)
+	{
+		if (bResetOnTeleport)
+		{
+			SimulatedTranslation = ComponentTransform.GetTranslation();
+			SimulatedRotation = ComponentTransform.GetRotation();
+		}
+		else
+		{
+			const FTransform OffsetTransform = FTransform(SimulatedRotation, SimulatedTranslation).GetRelativeTransformReverse(LastComponentTransform);
+			const FTransform TeleportedTransform = OffsetTransform * ComponentTransform;
+			SimulatedTranslation = TeleportedTransform.GetLocation();
+			SimulatedRotation = TeleportedTransform.GetRotation();
+		}
+
+		LastComponentTransform = ComponentTransform;
+	}
+
+	// Cycle last transform
+	LastMeshComponentTransformWorld = MeshComponentTransformWorld;
+
+	const EUAFOffsetRootBoneNodeMode CurrentTranslationMode = TranslationMode;
+	const EUAFOffsetRootBoneNodeMode CurrentRotationMode = RotationMode;
+
+	const bool bShouldConsumeTranslationOffset = UE::UAF::OffsetRootBone::ShouldExtractRootMotion(CurrentTranslationMode);
+	const bool bShouldConsumeRotationOffset = UE::UAF::OffsetRootBone::ShouldExtractRootMotion(CurrentRotationMode);
+
+	RootMotionTransformDelta.NormalizeRotation();
+
+	FTransform ConsumedRootMotionDelta = FTransform::Identity;
+
+	if (bShouldConsumeTranslationOffset)
+	{
+		ConsumedRootMotionDelta.SetTranslation(RootMotionTransformDelta.GetTranslation());
+	}
+	if (bShouldConsumeRotationOffset)
+	{
+		ConsumedRootMotionDelta.SetRotation(RootMotionTransformDelta.GetRotation());
+	}
+
+	if (UE::UAF::OffsetRootBone::ShouldCounterComponentDelta(CurrentRotationMode))
+	{
+		// Accumulate the rotation component delta into the simulated rotation
+		const FQuat ComponentRotationDelta = LastComponentTransform.GetRotation().Inverse() * ComponentTransform.GetRotation();
+		SimulatedRotation = ComponentRotationDelta * SimulatedRotation;
+	}
+	if (UE::UAF::OffsetRootBone::ShouldCounterComponentDelta(CurrentTranslationMode))
+	{
+		// Accumulate the translation component delta into the simulated translation
+		const FVector ComponentTranslationDelta = ComponentTransform.GetLocation() - LastComponentTransform.GetLocation();
+		SimulatedTranslation += ComponentTranslationDelta;
+	}
+
+	FTransform SimulatedTransform(SimulatedRotation, SimulatedTranslation);
+	// Apply the root motion delta
+	SimulatedTransform = ConsumedRootMotionDelta * SimulatedTransform;
+
+	SimulatedTranslation = SimulatedTransform.GetLocation();
+	SimulatedRotation = SimulatedTransform.GetRotation();
+
+	if (bOnGround)
+	{
+		SimulatedTranslation = FVector::PointPlaneProject(SimulatedTranslation, ComponentTransform.GetLocation(), AnimatedGroundNormal);
+	}
+
+	if (TranslationMode == EUAFOffsetRootBoneNodeMode::Release ||
+		TranslationMode == EUAFOffsetRootBoneNodeMode::Interpolate)
+	{
+		FVector TranslationOffset = ComponentTransform.GetLocation() - SimulatedTranslation;
+
+		// Blend out translation offset
+		FVector TranslationOffsetDelta = FVector::ZeroVector;
+		FMath::ExponentialSmoothingApprox(TranslationOffsetDelta, TranslationOffset, DeltaTime, TranslationSmoothingTime);
+
+		if (bClampToTranslationVelocity)
+		{
+			const float RootMotionDelta = RootMotionTransformDelta.GetLocation().Size();
+			const float MaxDelta = TranslationSpeedRatio * RootMotionDelta;
+
+			const float AdjustmentDelta = TranslationOffsetDelta.Size();
+			if (AdjustmentDelta > MaxDelta)
+			{
+				TranslationOffsetDelta = MaxDelta * TranslationOffsetDelta.GetSafeNormal2D();
+			}
+		}
+
+		SimulatedTranslation = SimulatedTranslation + TranslationOffsetDelta;
+	}
+
+	if (RotationMode == EUAFOffsetRootBoneNodeMode::Release ||
+		RotationMode == EUAFOffsetRootBoneNodeMode::Interpolate)
+	{
+		FQuat RotationOffset = ComponentTransform.GetRotation() * SimulatedRotation.Inverse();
+		RotationOffset.EnforceShortestArcWith(FQuat::Identity);
+		FQuat RotationOffsetDelta = FQuat::Identity;
+		SpringMath::ExponentialSmoothingApproxQuat(RotationOffsetDelta, RotationOffset, DeltaTime, RotationSmoothingTime);
+
+		if (bClampToRotationVelocity)
+		{
+			float RotationMotionAngleDelta;
+			FVector RootMotionRotationAxis;
+			RootMotionTransformDelta.GetRotation().ToAxisAndAngle(RootMotionRotationAxis, RotationMotionAngleDelta);
+
+			float MaxRotationAngle = RotationSpeedRatio * RotationMotionAngleDelta;
+
+			FVector DeltaAxis;
+			float DeltaAngle;
+			RotationOffsetDelta.ToAxisAndAngle(DeltaAxis, DeltaAngle);
+
+			if (DeltaAngle > MaxRotationAngle)
+			{
+				RotationOffsetDelta = FQuat(DeltaAxis, MaxRotationAngle);
+			}
+		}
+
+		SimulatedRotation = RotationOffsetDelta * SimulatedRotation;
+	}
+
+	if (MaxTranslationError >= 0.0f)
+	{
+		FVector TranslationOffset = ComponentTransform.GetLocation() - SimulatedTranslation;
+		const float TranslationOffsetSizeSquared = TranslationOffset.SizeSquared();
+		if (TranslationOffsetSizeSquared > (MaxTranslationError * MaxTranslationError))
+		{
+			TranslationOffset = TranslationOffset.GetClampedToMaxSize(MaxTranslationError);
+			SimulatedTranslation = ComponentTransform.GetLocation() - TranslationOffset;
+		}
+	}
+
+	const float MaxAngleRadians = FMath::DegreesToRadians(MaxRotationErrorDegrees);
+	if (MaxAngleRadians >= 0.0f)
+	{
+		FQuat RotationOffset = ComponentTransform.GetRotation().Inverse() * SimulatedRotation;
+		RotationOffset.EnforceShortestArcWith(FQuat::Identity);
+
+		FVector OffsetAxis;
+		float OffsetAngle;
+		RotationOffset.ToAxisAndAngle(OffsetAxis, OffsetAngle);
+
+		if (FMath::Abs(OffsetAngle) > MaxAngleRadians)
+		{
+			RotationOffset = FQuat(OffsetAxis, MaxAngleRadians);
+			SimulatedRotation = RotationOffset * ComponentTransform.GetRotation();
+			SimulatedRotation.Normalize();
+		}
+	}
+
+	// Apply the offset adjustments to the simulated transform
+	SimulatedTransform.SetLocation(SimulatedTranslation);
+	SimulatedTransform.SetRotation(SimulatedRotation);
+
+	// Start with the input pose's bone transform, to preserve any adjustments done before this node in the graph
+	FTransform TargetBoneTransform = InputBoneTransform;
+	// Accumulate the simulated transform in, and counter current component transform
+	TargetBoneTransform.Accumulate(SimulatedTransform * ComponentTransform.Inverse());
+
+	// Offset root bone should not affect scale so take the input
+	TargetBoneTransform.SetScale3D(InputBoneTransform.GetScale3D());
+
+	TargetBoneTransform.NormalizeRotation();
+	check(SimulatedTranslation.ContainsNaN() == false);
+
+	// Write modified root bone transform back
+	(*BoneMap)[RootBoneSetIndex].Value = TargetBoneTransform;
+
+#if ENABLE_VISUAL_LOG && ENABLE_ANIM_DEBUG
+	if (FVisualLogger::IsRecording())
+	{
+		static const TCHAR* LogName = TEXT("OffsetRootBone");
+		const float InnerCircleRadius = 40.0f;
+		const uint16 CircleThickness = 2;
+		const FVector CircleOffset(0, 0, 1);
+
+		const FTransform TargetBoneInitialTransformWorld = InputBoneTransform * ComponentTransform;
+		const FTransform TargetBoneTransformWorld = TargetBoneTransform * ComponentTransform;
+
+		if (MaxTranslationError >= 0.0f)
+		{
+			const float OuterCircleRadius = MaxTranslationError + InnerCircleRadius;
+			UE_VLOG_CIRCLE_THICK(HostObject, TEXT("OffsetRootBone"), Display, ComponentTransform.GetLocation() + CircleOffset, FVector::UpVector, OuterCircleRadius, FColor::Red, CircleThickness, TEXT(""));
+		}
+
+		UE_VLOG_CIRCLE_THICK(HostObject, LogName, Display, ComponentTransform.GetLocation() + CircleOffset, FVector::UpVector, InnerCircleRadius, FColor::Blue, CircleThickness, TEXT(""));
+		UE_VLOG_ARROW(HostObject, LogName, Display,
+			ComponentTransform.GetLocation() + CircleOffset,
+			ComponentTransform.GetLocation() + InnerCircleRadius * ComponentTransform.GetRotation().GetRightVector() + CircleOffset,
+			FColor::Blue, TEXT(""));
+
+		UE_VLOG_CIRCLE_THICK(HostObject, LogName, Display, TargetBoneTransformWorld.GetLocation() + CircleOffset, FVector::UpVector, InnerCircleRadius, FColor::Green, CircleThickness, TEXT(""));
+		UE_VLOG_ARROW(HostObject, LogName, Display,
+			TargetBoneTransformWorld.GetLocation() + CircleOffset,
+			TargetBoneTransformWorld.GetLocation() + InnerCircleRadius * TargetBoneTransformWorld.GetRotation().GetRightVector() + CircleOffset,
+			FColor::Green, TEXT(""));
+	}
+#endif // ENABLE_VISUAL_LOG && ENABLE_ANIM_DEBUG
+}
+
+} // namespace UE::UAF
diff --git a/Engine/Plugins/Experimental/UAF/UAFAnimNode/Source/UAFAnimNode/Public/UAF/AnimNodes/UAFOffsetRootBoneNode.h b/Engine/Plugins/Experimental/UAF/UAFAnimNode/Source/UAFAnimNode/Public/UAF/AnimNodes/UAFOffsetRootBoneNode.h
new file mode 100644
--- /dev/null
+++ b/Engine/Plugins/Experimental/UAF/UAFAnimNode/Source/UAFAnimNode/Public/UAF/AnimNodes/UAFOffsetRootBoneNode.h
@@ -0,0 +1,128 @@
+// Copyright Epic Games, Inc. All Rights Reserved.
+
+#pragma once
+
+#include "CoreMinimal.h"
+#include "BindableValue/UAFBindableTypes.h"
+#include "UAF/AnimNodeCore/UAFAnimNode.h"
+#include "UAF/AnimNodeCore/UAFAnimNodeData.h"
+#include "UAF/AnimOps/UAFOffsetRootBoneAnimOp.h"
+#include "UAFOffsetRootBoneNode.generated.h"
+
+#define UE_API UAFANIMNODE_API
+
+namespace UE::UAF
+{
+
+/**
+ * Adds a procedural offset to the root bone coming from incoming root motion.
+ * It can optionally fade this out over time (controlled by EUAFOffsetRootBoneNodeMode).
+ *
+ * This allows a code-controlled character to still be influenced by root motion,
+ * e.g. during plant turns. It can lead to less foot sliding in a motion matching setup.
+ */
+USTRUCT(DisplayName = "Offset Root Bone")
+struct FUAFOffsetRootBoneNodeData : public FUAFAnimNodeData
+{
+	GENERATED_BODY()
+
+	// Input pose to be processed
+	UPROPERTY(EditAnywhere, Category = "Data")
+	TInstancedStruct<FUAFAnimNodeData> Child;
+
+	// Current strength of the skeletal control
+	UPROPERTY(EditAnywhere, Category = Alpha)
+	FBindableFloat Alpha = 1.0f;
+
+	// The translation offset behavior mode
+	UPROPERTY(EditAnywhere, Category = Settings)
+	FBindableEnum TranslationMode = FBindableEnum(EUAFOffsetRootBoneNodeMode::Interpolate);
+
+	// The rotation offset behavior mode
+	UPROPERTY(EditAnywhere, Category = Settings)
+	FBindableEnum RotationMode = FBindableEnum(EUAFOffsetRootBoneNodeMode::Interpolate);
+
+	// Controls how fast the translation offset is blended out
+	// Values closer to 0 make it faster
+	UPROPERTY(EditAnywhere, Category = Settings)
+	FBindableFloat TranslationSmoothingTime = 0.1f;
+
+	// Controls how fast the rotation offset is blended out
+	// Values closer to 0 make it faster
+	UPROPERTY(EditAnywhere, Category = Settings)
+	FBindableFloat RotationSmoothingTime = 0.2f;
+
+	// How much the offset can deviate from the mesh component's translation in units
+	// Values lower than 0 disable this limit
+	UPROPERTY(EditAnywhere, Category = Settings)
+	FBindableFloat MaxTranslationError = -1.0f;
+
+	// How much the offset can deviate from the mesh component's rotation in degrees
+	// Values lower than 0 disable this limit
+	UPROPERTY(EditAnywhere, Category = Settings)
+	FBindableFloat MaxRotationErrorDegrees = -1.0f;
+
+	// Whether to limit the offset's translation interpolation speed to the velocity on the incoming motion
+	// Enabling this prevents the offset sliding when there's little to no translation speed
+	UPROPERTY(EditAnywhere, Category = Settings, meta = (Inline))
+	bool bClampToTranslationVelocity = false;
+
+	// Whether to limit the offset's rotation interpolation speed to the velocity on the incoming motion
+	// Enabling this prevents the offset sliding when there's little to no rotation speed
+	UPROPERTY(EditAnywhere, Category = Settings, meta = (Inline))
+	bool bClampToRotationVelocity = false;
+
+	// How much the offset can blend out, relative to the incoming translation speed
+	// i.e. If root motion is moving at 400cm/s, at 0.5, the offset can blend out at 200cm/s
+	UPROPERTY(EditAnywhere, Category = Settings, meta = (EditCondition = bClampToTranslationVelocity, Inline))
+	float TranslationSpeedRatio = 0.5f;
+
+	// How much the offset can blend out, relative to the incoming rotation speed
+	// i.e. If root motion is rotating at 90 degrees/s, at 0.5, the offset can blend out at 45 degree/s
+	UPROPERTY(EditAnywhere, Category = Settings, meta = (EditCondition = bClampToRotationVelocity, Inline))
+	float RotationSpeedRatio = 0.5f;
+
+	// This has to be the mesh component transform in world space
+	UPROPERTY(EditAnywhere, Category = Evaluation)
+	FBindableTransform MeshComponentTransformWorld;
+
+	// Signals whether the offset root bone translation should be projected onto a projected ground plane
+	UPROPERTY(EditAnywhere, Category = Settings)
+	FBindableBool bOnGround = true;
+
+	// Surface normal that the current animation was authored for. Used with the bOnGround flag.
+	UPROPERTY(EditAnywhere, Category = Settings)
+	FBindableVector AnimatedGroundNormal = FVector(0, 0, 1);
+
+	// Distance threshold used to detect teleports
+	UPROPERTY(EditAnywhere, Category = Settings, meta = (Inline))
+	float TeleportDistanceThreshold = 300.0f;
+
+	// If to reset the root bone offset on teleport
+	UPROPERTY(EditAnywhere, Category = Settings, meta = (Inline))
+	bool bResetOnTeleport = false;
+
+	virtual FUAFAnimNodePtr CreateInstance(FUAFAnimGraphUpdateContext& Context) const override;
+};
+
+class FUAFOffsetRootBoneNode : public FUAFAnimNode
+{
+public:
+	UE_API FUAFOffsetRootBoneNode(FUAFAnimGraphUpdateContext& Context, const FUAFOffsetRootBoneNodeData& InData);
+
+#if UAF_TRACE_ENABLED
+	virtual FString GetDebugName() const override;
+	virtual UStruct* GetDebugStruct() const override;
+#endif
+
+protected:
+	virtual void PreUpdate(FUAFAnimGraphUpdateContext& GraphContext) override;
+
+	const FUAFOffsetRootBoneNodeData* Data;
+
+	FUAFOffsetRootBoneAnimOp OffsetRootBoneAnimOp;
+};
+
+} // namespace UE::UAF
+
+#undef UE_API
diff --git a/Engine/Plugins/Experimental/UAF/UAFAnimNode/Source/UAFAnimNode/Public/UAF/AnimOps/UAFOffsetRootBoneAnimOp.h b/Engine/Plugins/Experimental/UAF/UAFAnimNode/Source/UAFAnimNode/Public/UAF/AnimOps/UAFOffsetRootBoneAnimOp.h
new file mode 100644
--- /dev/null
+++ b/Engine/Plugins/Experimental/UAF/UAFAnimNode/Source/UAFAnimNode/Public/UAF/AnimOps/UAFOffsetRootBoneAnimOp.h
@@ -0,0 +1,108 @@
+// Copyright Epic Games, Inc. All Rights Reserved.
+
+#pragma once
+
+#include "CoreMinimal.h"
+#include "UAF/AnimOpCore/UAFAnimOp.h"
+#include "UAFOffsetRootBoneAnimOp.generated.h"
+
+#define UE_API UAFANIMNODE_API
+
+UENUM(BlueprintType)
+enum class EUAFOffsetRootBoneNodeMode : uint8
+{
+	// Accumulate the mesh component's movement into the offset.
+	// In this mode, if the mesh component moves
+	// the offset will counter the motion, and the root will stay in place
+	Accumulate,
+	// Continuously interpolate the offset out
+	// In this mode, if the mesh component moves
+	// The root will stay behind, but will attempt to catch up
+	Interpolate,
+	// Release the offset and stop accumulating the mesh component's movement delta.
+	// In this mode we will "blend out" the offset
+	Release,
+};
+
+namespace UE::UAF
+{
+
+USTRUCT()
+struct FUAFOffsetRootBoneAnimOp : public FUAFAnimOp
+{
+	GENERATED_BODY()
+	UAF_DECLARE_ANIMOP(FUAFOffsetRootBoneAnimOp)
+
+	FUAFOffsetRootBoneAnimOp();
+
+	// FUAFAnimOp impl
+	UE_API virtual void EvaluateValues(FUAFAnimOpValueEvaluator& Evaluator) override;
+
+	// Per-frame resolved values (set by Node's PostUpdate)
+
+	UPROPERTY()
+	float DeltaTime = 0.f;
+
+	UPROPERTY()
+	float Alpha = 1.0f;
+
+	UPROPERTY()
+	FTransform MeshComponentTransformWorld = FTransform::Identity;
+
+	UPROPERTY()
+	EUAFOffsetRootBoneNodeMode TranslationMode = EUAFOffsetRootBoneNodeMode::Interpolate;
+
+	UPROPERTY()
+	EUAFOffsetRootBoneNodeMode RotationMode = EUAFOffsetRootBoneNodeMode::Interpolate;
+
+	UPROPERTY()
+	float TranslationSmoothingTime = 0.1f;
+
+	UPROPERTY()
+	float RotationSmoothingTime = 0.2f;
+
+	UPROPERTY()
+	float MaxTranslationError = -1.0f;
+
+	UPROPERTY()
+	float MaxRotationErrorDegrees = -1.0f;
+
+	UPROPERTY()
+	bool bOnGround = true;
+
+	UPROPERTY()
+	FVector AnimatedGroundNormal = FVector(0, 0, 1);
+
+	// Constant values (set once in Node constructor)
+
+	UPROPERTY()
+	bool bClampToTranslationVelocity = false;
+
+	UPROPERTY()
+	bool bClampToRotationVelocity = false;
+
+	UPROPERTY()
+	float TranslationSpeedRatio = 0.5f;
+
+	UPROPERTY()
+	float RotationSpeedRatio = 0.5f;
+
+	UPROPERTY()
+	float TeleportDistanceThreshold = 300.0f;
+
+	UPROPERTY()
+	bool bResetOnTeleport = false;
+
+	// Persistent state (survives across frames)
+
+	bool bIsFirstUpdate = true;
+	FTransform LastMeshComponentTransformWorld = FTransform::Identity;
+	FVector SimulatedTranslation = FVector::ZeroVector;
+	FQuat SimulatedRotation = FQuat::Identity;
+
+	UObject* HostObject = nullptr;
+};
+
+} // namespace UE::UAF
+
+#undef UE_API
