Fix thread safety of frame provider access by scoping reads with FAnalysisSessionReadScope
Refactor UAF anim node track to flatten nesting, add early returns, and properly scope provider read locks

#Synced-CL 52294509

diff --git a/Engine/Plugins/Animation/GameplayInsights/Source/GameplayInsights/Private/SAnimGraphSchematicView.cpp b/Engine/Plugins/Animation/GameplayInsights/Source/GameplayInsights/Private/SAnimGraphSchematicView.cpp
--- a/Engine/Plugins/Animation/GameplayInsights/Source/GameplayInsights/Private/SAnimGraphSchematicView.cpp
+++ b/Engine/Plugins/Animation/GameplayInsights/Source/GameplayInsights/Private/SAnimGraphSchematicView.cpp
@@ -555,10 +555,15 @@ void SAnimGraphSchematicView::RefreshNodes()
 
 		AnimationProvider->ReadAnimGraphTimeline(AnimInstanceId, [this, &NodeMap, AnimationProvider, GameplayProvider](const FAnimationProvider::AnimGraphTimeline& InGraphTimeline)
 		{
-			const TraceServices::IFrameProvider& FramesProvider = TraceServices::ReadFrameProvider(*AnalysisSession);
-
+			bool bFrameFound;
 			TraceServices::FFrame Frame;
-			if(FramesProvider.GetFrameFromTime(ETraceFrameType::TraceFrameType_Game, TimeMarker, Frame))
+			{
+				TraceServices::FAnalysisSessionReadScope SessionReadScope(*AnalysisSession);
+				const TraceServices::IFrameProvider& FrameProvider = TraceServices::ReadFrameProvider(*AnalysisSession);
+				bFrameFound = FrameProvider.GetFrameFromTime(TraceFrameType_Game, TimeMarker, Frame);
+			}
+
+			if (bFrameFound)
 			{
 				InGraphTimeline.EnumerateEvents(Frame.StartTime, Frame.EndTime, [this, &NodeMap, AnimationProvider, GameplayProvider](double InGraphStartTime, double InGraphEndTime, uint32 InDepth, const FAnimGraphMessage& InMessage)
 				{
diff --git a/Engine/Plugins/Experimental/UAF/UAFAnimNode/Source/UAFAnimNodeEditor/Private/RewindDebugger/UAFAnimNodeTrack.cpp b/Engine/Plugins/Experimental/UAF/UAFAnimNode/Source/UAFAnimNodeEditor/Private/RewindDebugger/UAFAnimNodeTrack.cpp
--- a/Engine/Plugins/Experimental/UAF/UAFAnimNode/Source/UAFAnimNodeEditor/Private/RewindDebugger/UAFAnimNodeTrack.cpp
+++ b/Engine/Plugins/Experimental/UAF/UAFAnimNode/Source/UAFAnimNodeEditor/Private/RewindDebugger/UAFAnimNodeTrack.cpp
@@ -141,134 +141,145 @@ bool FUAFAnimNodeTrack::UpdateInternal()
 	const IGameplayProvider* GameplayProvider = Session->ReadProvider<IGameplayProvider>("GameplayProvider");
 	const IAnimationProvider* AnimationProvider = Session->ReadProvider<IAnimationProvider>("AnimationProvider");
 	const FUAFAnimNodeProvider* UAFAnimNodeProvider = Session->ReadProvider<FUAFAnimNodeProvider>("UAFAnimNodeProvider");
-	const TraceServices::IFrameProvider& FramesProvider = TraceServices::ReadFrameProvider(*Session);
+	if (GameplayProvider == nullptr
+		|| AnimationProvider == nullptr
+		|| UAFAnimNodeProvider == nullptr)
+	{
+		return false;
+	}
 
-	if (GameplayProvider && AnimationProvider && UAFAnimNodeProvider)
+	const double CurrentScrubTime = RewindDebugger->CurrentTraceTime();
+
+	TraceServices::FFrame MarkerFrame;
 	{
-		TraceServices::FProviderReadScopeLock GameplayProviderReadScope(*GameplayProvider);
-		TraceServices::FProviderReadScopeLock AnimationProviderReadScope(*AnimationProvider);
-		TraceServices::FProviderReadScopeLock UAFAnimNodeProviderReadScope(*UAFAnimNodeProvider);
-		double CurrentScrubTime = RewindDebugger->CurrentTraceTime();
+		const TraceServices::IFrameProvider& FrameProvider = TraceServices::ReadFrameProvider(*RewindDebugger->GetAnalysisSession());
+		TraceServices::FAnalysisSessionReadScope SessionReadScope(*Session);
+		if (!FrameProvider.GetFrameFromTime(ETraceFrameType::TraceFrameType_Game, CurrentScrubTime, MarkerFrame))
+		{
+			return false;
+		}
+	}
 
+	if (GraphId == 0)
+	{
+		TraceServices::FProviderReadScopeLock GameplayProviderReadScope(*GameplayProvider);
+		const FObjectInfo& Info = GameplayProvider->GetObjectInfo(NodeId);
+		GraphId = Info.GetOuterId().GetMainId();
 		if (GraphId == 0)
 		{
-			const FObjectInfo& Info = GameplayProvider->GetObjectInfo(NodeId);
-			GraphId = Info.GetOuterId().GetMainId();
+			return false;
 		}
+	}
 
-		if (GraphId != 0)
+	TraceServices::FProviderReadScopeLock UAFAnimNodeProviderReadScope(*UAFAnimNodeProvider);
+	if (const FAnimNodeUpdateTimelineData* Data = UAFAnimNodeProvider->GetAnimNodeTimelineData(GraphId))
+	{
+		if (PreviousScrubTime != CurrentScrubTime && DetailsObject)
 		{
-			if (const FAnimNodeUpdateTimelineData* Data = UAFAnimNodeProvider->GetAnimNodeTimelineData(GraphId))
-			{
-				if (PreviousScrubTime != CurrentScrubTime && DetailsObject)
+			PreviousScrubTime = CurrentScrubTime;
+
+			// clear properties first so we don't see old data
+			DetailsObject->Properties.Reset();
+
+			Data->NodeUpdateTimeline.EnumerateEvents(MarkerFrame.StartTime, MarkerFrame.EndTime,
+				[this](double InStartTime, double InEndTime, uint32 InDepth, const FAnimNodeUpdateData& NodeUpdateData)
 				{
-					PreviousScrubTime = CurrentScrubTime;
+					if (NodeUpdateData.NodeId == NodeId)
+					{
+						static FName PropertyName("TotalWeight");
+						FPropertyBagPropertyDesc PropertyDesc(PropertyName, EPropertyBagPropertyType::Float);
+						DetailsObject->Properties.AddProperties({ PropertyDesc });
+						DetailsObject->Properties.SetValueFloat(PropertyName, NodeUpdateData.TotalWeight);
+						return TraceServices::EEventEnumerate::Stop;
+					}
+					return TraceServices::EEventEnumerate::Continue;
+				});
 
-					// clear properties first so we don't see old data
-					DetailsObject->Properties.Reset();
+			if (auto ValueData = UAFAnimNodeProvider->GetAnimNodeValueTimelineData(GraphId))
+			{
+				TraceServices::FProviderReadScopeLock GameplayProviderReadScope(*GameplayProvider);
+				TraceServices::FProviderReadScopeLock AnimationProviderReadScope(*AnimationProvider);
 
-					TraceServices::FFrame MarkerFrame;
-					if (FramesProvider.GetFrameFromTime(ETraceFrameType::TraceFrameType_Game, CurrentScrubTime, MarkerFrame))
+				ValueData->NodeValueTimeline.EnumerateEvents(MarkerFrame.StartTime, MarkerFrame.EndTime,
+					[GameplayProvider, AnimationProvider, this](double InStartTime, double InEndTime, uint32 InDepth, const FAnimNodeValueData& NodeValueData)
 					{
-						Data->NodeUpdateTimeline.EnumerateEvents(MarkerFrame.StartTime, MarkerFrame.EndTime,
-						[GameplayProvider, this] ( double InStartTime, double InEndTime, uint32 InDepth, const FAnimNodeUpdateData& NodeUpdateData)
+						if (NodeValueData.NodeId == NodeId)
 						{
-							if (NodeUpdateData.NodeId == NodeId)
-							{
-								static FName PropertyName("TotalWeight");
-								FPropertyBagPropertyDesc PropertyDesc(PropertyName,  EPropertyBagPropertyType::Float);
-								DetailsObject->Properties.AddProperties({PropertyDesc});
-								DetailsObject->Properties.SetValueFloat(PropertyName, NodeUpdateData.TotalWeight);	
-								return TraceServices::EEventEnumerate::Stop;
-							}
-							return TraceServices::EEventEnumerate::Continue;
-						});
+							FName PropertyName = AnimationProvider->GetName(NodeValueData.NameId);
+							FPropertyBagPropertyDesc PropertyDesc(PropertyName, static_cast<EPropertyBagPropertyType>(NodeValueData.Type), UObject::StaticClass());
 
-						if (auto ValueData = UAFAnimNodeProvider->GetAnimNodeValueTimelineData(GraphId))
-						{
-							ValueData->NodeValueTimeline.EnumerateEvents(MarkerFrame.StartTime, MarkerFrame.EndTime,
-							[GameplayProvider, AnimationProvider, this] ( double InStartTime, double InEndTime, uint32 InDepth, const FAnimNodeValueData& NodeValueData)
+							DetailsObject->Properties.AddProperties({ PropertyDesc });
+
+							switch (static_cast<EPropertyBagPropertyType>(NodeValueData.Type))
+							{
+							case EPropertyBagPropertyType::Bool:
+								if (ensure(NodeValueData.Value.Num() == sizeof(bool)))
+								{
+									DetailsObject->Properties.SetValueBool(PropertyName, *reinterpret_cast<const bool*>(&NodeValueData.Value[0]));
+								}
+								break;
+							case EPropertyBagPropertyType::Double:
+								if (ensure(NodeValueData.Value.Num() == sizeof(double)))
+								{
+									DetailsObject->Properties.SetValueDouble(PropertyName, *reinterpret_cast<const double*>(&NodeValueData.Value[0]));
+								}
+								break;
+							case EPropertyBagPropertyType::Object:
 							{
-								if (NodeValueData.NodeId == NodeId)
+								if (ensure(NodeValueData.Value.Num() == sizeof(uint64)))
 								{
-									FName PropertyName = AnimationProvider->GetName(NodeValueData.NameId);
-									FPropertyBagPropertyDesc PropertyDesc(PropertyName,  static_cast<EPropertyBagPropertyType>(NodeValueData.Type), UObject::StaticClass());
+									UObject* Asset = nullptr;
+									const uint64 ObjectId = *reinterpret_cast<const uint64*>(&NodeValueData.Value[0]);
+									const FObjectInfo& AssetInfo = GameplayProvider->GetObjectInfo(ObjectId);
+									const FString PackagePathString = FPackageName::ObjectPathToPackageName(FString(AssetInfo.PathName));
+
+									UPackage* Package = LoadPackage(nullptr, ToCStr(PackagePathString), LOAD_NoRedirects);
+									if (Package)
+									{
+										Package->FullyLoad();
 
-									DetailsObject->Properties.AddProperties({PropertyDesc});
-									
-									switch (static_cast<EPropertyBagPropertyType>(NodeValueData.Type))
+										FString AssetName = FPaths::GetBaseFilename(AssetInfo.PathName);
+										Asset = FindObject<UObject>(Package, *AssetName);
+									}
+									else
 									{
-										case EPropertyBagPropertyType::Bool:
-											if(ensure(NodeValueData.Value.Num() == sizeof(bool)))
-											{
-												DetailsObject->Properties.SetValueBool(PropertyName, *reinterpret_cast<const bool*>(&NodeValueData.Value[0]));
-											}
-										break;
-										case EPropertyBagPropertyType::Double:
-											if(ensure(NodeValueData.Value.Num() == sizeof(double)))
-											{
-												DetailsObject->Properties.SetValueDouble(PropertyName, *reinterpret_cast<const double*>(&NodeValueData.Value[0]));
-											}
-										break;
-										case EPropertyBagPropertyType::Object:
-										{
-											if(ensure(NodeValueData.Value.Num() == sizeof(uint64)))
-											{
-												UObject* Asset = nullptr;
-												uint64 ObjectId = *reinterpret_cast<const uint64*>(&NodeValueData.Value[0]);
-												const FObjectInfo& AssetInfo = GameplayProvider->GetObjectInfo(ObjectId);
-												FString PackagePathString = FPackageName::ObjectPathToPackageName(FString(AssetInfo.PathName));
-
-												UPackage* Package = LoadPackage(NULL, ToCStr(PackagePathString), LOAD_NoRedirects);
-												if (Package)
-												{
-													Package->FullyLoad();
-													
-													FString AssetName = FPaths::GetBaseFilename(AssetInfo.PathName);
-													Asset = FindObject<UObject>(Package, *AssetName);
-												}
-												else
-												{
-													// fallback for unsaved assets
-													Asset = FindObject<UObject>(nullptr, AssetInfo.PathName);
-												}
-
-												DetailsObject->Properties.SetValueObject(PropertyName, Asset);
-											}
-											break;
-										}
-										default:
-										break;
-									};
+										// fallback for unsaved assets
+										Asset = FindObject<UObject>(nullptr, AssetInfo.PathName);
+									}
+
+									DetailsObject->Properties.SetValueObject(PropertyName, Asset);
 								}
-								return TraceServices::EEventEnumerate::Continue;
-							});
+								break;
+							}
+							default:
+								break;
+							};
 						}
-					}
-				}
+						return TraceServices::EEventEnumerate::Continue;
+					});
+			}
+		}
 
-				if(CurvesUpdateRequested > 10)
-				{
-					TRACE_CPUPROFILER_EVENT_SCOPE(FUAFAnimNodeTrack::UpdateCurvePointsInternal);
-					CurveData->Points.SetNum(0,EAllowShrinking::No);
+		if (CurvesUpdateRequested > 10)
+		{
+			TRACE_CPUPROFILER_EVENT_SCOPE(FUAFAnimNodeTrack::UpdateCurvePointsInternal);
+			CurveData->Points.SetNum(0, EAllowShrinking::No);
 
-					TRange<double> TraceTimeRange = RewindDebugger->GetCurrentTraceRange();
-					double StartTime = TraceTimeRange.GetLowerBoundValue();
-					double EndTime = TraceTimeRange.GetUpperBoundValue();
+			const TRange<double> TraceTimeRange = RewindDebugger->GetCurrentTraceRange();
+			const double StartTime = TraceTimeRange.GetLowerBoundValue();
+			const double EndTime = TraceTimeRange.GetUpperBoundValue();
 
-					Data->NodeUpdateTimeline.EnumerateEvents(StartTime, EndTime, [this](double InStartTime, double InEndTime, uint32 InDepth, const FAnimNodeUpdateData& NodeUpdateData)
+			Data->NodeUpdateTimeline.EnumerateEvents(StartTime, EndTime, [this](double InStartTime, double InEndTime, uint32 InDepth, const FAnimNodeUpdateData& NodeUpdateData)
+				{
+					if (NodeUpdateData.NodeId == NodeId)
 					{
-						if (NodeUpdateData.NodeId == NodeId)
-						{
-							CurveData->Points.Add({ NodeUpdateData.RecordingTime,	NodeUpdateData.TotalWeight });
-						}
-										
-						return TraceServices::EEventEnumerate::Continue;
-					});
+						CurveData->Points.Add({ NodeUpdateData.RecordingTime, NodeUpdateData.TotalWeight });
+					}
 
-					CurvesUpdateRequested = 0;
-				}
-			}
+					return TraceServices::EEventEnumerate::Continue;
+				});
+
+			CurvesUpdateRequested = 0;
 		}
 	}
 
