UAFAnimNode - added missing reinitialization for DetailsObject in FUAFSequenceInfoTrack::UpdateInternal

#Synced-CL 52316895

diff --git a/Engine/Plugins/Experimental/UAF/UAFAnimNode/Source/UAFAnimNodeEditor/Private/RewindDebugger/UAFSequenceInfoTrack.cpp b/Engine/Plugins/Experimental/UAF/UAFAnimNode/Source/UAFAnimNodeEditor/Private/RewindDebugger/UAFSequenceInfoTrack.cpp
--- a/Engine/Plugins/Experimental/UAF/UAFAnimNode/Source/UAFAnimNodeEditor/Private/RewindDebugger/UAFSequenceInfoTrack.cpp
+++ b/Engine/Plugins/Experimental/UAF/UAFAnimNode/Source/UAFAnimNodeEditor/Private/RewindDebugger/UAFSequenceInfoTrack.cpp
@@ -294,6 +294,7 @@ bool FUAFSequenceInfoTrack::UpdateInternal()
 		{
 			// this should not happen unless the object was garbage collected (which should not happen since it's marked as Standalone)
 			Initialize();
+			DetailsObject = DetailsObjectWeakPtr.Get();
 		}
 
 		if (InstanceId == 0)
