Fix notify queries using current time instead of delta time in new UAF runtime

#Synced-CL 52358354

diff --git a/Engine/Plugins/Experimental/UAF/UAFAnimNode/Source/UAFAnimNode/Private/RigUnit_RunAnimNode.cpp b/Engine/Plugins/Experimental/UAF/UAFAnimNode/Source/UAFAnimNode/Private/RigUnit_RunAnimNode.cpp
--- a/Engine/Plugins/Experimental/UAF/UAFAnimNode/Source/UAFAnimNode/Private/RigUnit_RunAnimNode.cpp
+++ b/Engine/Plugins/Experimental/UAF/UAFAnimNode/Source/UAFAnimNode/Private/RigUnit_RunAnimNode.cpp
@@ -130,15 +130,17 @@ FRigUnit_RunAnimNode_Execute()
 		// Use AnimOps to evaluate our notifies
 		{
 			TArray<FAnimNotifyEventReference> Notifies = EvaluateNotifies(AnimOpList);
+			if (!Notifies.IsEmpty())
+			{
+				// Ensure we have a handler component on the module
+				(void)ModuleInstance.GetComponent<FUAFNotifyDispatcherComponent>();
 
-			// Ensure we have a handler component on the module
-			(void)ModuleInstance.GetComponent<FUAFNotifyDispatcherComponent>();
-
-			auto NotifyDispatchEvent = MakeTraitEvent<FNotifyDispatchEvent>();
-			NotifyDispatchEvent->Notifies = MoveTemp(Notifies);
-			NotifyDispatchEvent->Weight = 1.0f;
+				auto NotifyDispatchEvent = MakeTraitEvent<FNotifyDispatchEvent>();
+				NotifyDispatchEvent->Notifies = MoveTemp(Notifies);
+				NotifyDispatchEvent->Weight = 1.0f;
 
-			ModuleInstance.QueueOutputTraitEvent(NotifyDispatchEvent);
+				ModuleInstance.QueueOutputTraitEvent(NotifyDispatchEvent);
+			}
 		}
 
 		// Evaluate sync contributors now while the AnimOps are ready, they will be used in the next update
diff --git a/Engine/Plugins/Experimental/UAF/UAFAnimNode/Source/UAFAnimNode/Private/UAF/AnimOps/UAFDecompressAnimSequenceAnimOp.cpp b/Engine/Plugins/Experimental/UAF/UAFAnimNode/Source/UAFAnimNode/Private/UAF/AnimOps/UAFDecompressAnimSequenceAnimOp.cpp
--- a/Engine/Plugins/Experimental/UAF/UAFAnimNode/Source/UAFAnimNode/Private/UAF/AnimOps/UAFDecompressAnimSequenceAnimOp.cpp
+++ b/Engine/Plugins/Experimental/UAF/UAFAnimNode/Source/UAFAnimNode/Private/UAF/AnimOps/UAFDecompressAnimSequenceAnimOp.cpp
@@ -111,7 +111,7 @@ namespace UE::UAF
 			TickRecord.bLooping = bIsLooping;
 
 			FAnimNotifyContext NotifyContext(TickRecord);
-			AnimSequencePtr->GetAnimNotifies(PreviousTime, CurrentTime, NotifyContext);
+			AnimSequencePtr->GetAnimNotifies(PreviousTime, DeltaTimeRecord.Delta, NotifyContext);
 
 			Evaluator.GetNotifies().Append(NotifyContext.ActiveNotifies);
 		}
