Restore RigUnit property to avoid cook issues

#Synced-CL 52358354

diff --git a/Engine/Plugins/Experimental/UAF/UAFAnimNode/Source/UAFAnimNode/Internal/AnimNode/RigUnit_RunAnimNode.h b/Engine/Plugins/Experimental/UAF/UAFAnimNode/Source/UAFAnimNode/Internal/AnimNode/RigUnit_RunAnimNode.h
--- a/Engine/Plugins/Experimental/UAF/UAFAnimNode/Source/UAFAnimNode/Internal/AnimNode/RigUnit_RunAnimNode.h
+++ b/Engine/Plugins/Experimental/UAF/UAFAnimNode/Source/UAFAnimNode/Internal/AnimNode/RigUnit_RunAnimNode.h
@@ -45,9 +45,13 @@ struct FRigUnit_RunAnimNode : public FRigUnit_AnimNextBase
 	RIGVM_METHOD()
 	UE_API void Execute();
 
+	// The AnimNode to run (DO NOT USE, RETAINING IT TEMPORARILY)
+	UPROPERTY(EditAnywhere, Category = "Run", meta = (Input))
+	FUAFAnimNodeReference AnimNode;
+
 	// The anim node to run
 	UPROPERTY(EditAnywhere, Category = "Run", meta = (Input))
-	UE::UAF::FUAFAnimNodeDataEx AnimNode;
+	UE::UAF::FUAFAnimNodeDataEx RootAnimNode;
 
 	// LOD to run at. If this is -1 then the reference pose's source LOD will be used
 	UPROPERTY(EditAnywhere, Category = "Run", meta = (Input))
diff --git a/Engine/Plugins/Experimental/UAF/UAFAnimNode/Source/UAFAnimNode/Private/RigUnit_RunAnimNode.cpp b/Engine/Plugins/Experimental/UAF/UAFAnimNode/Source/UAFAnimNode/Private/RigUnit_RunAnimNode.cpp
--- a/Engine/Plugins/Experimental/UAF/UAFAnimNode/Source/UAFAnimNode/Private/RigUnit_RunAnimNode.cpp
+++ b/Engine/Plugins/Experimental/UAF/UAFAnimNode/Source/UAFAnimNode/Private/RigUnit_RunAnimNode.cpp
@@ -41,7 +41,7 @@ FRigUnit_RunAnimNode_Execute()
 		return;
 	}
 	
-	if (!AnimNode.IsValid())
+	if (!RootAnimNode.IsValid())
 	{
 		UAF_RIGUNIT_LOG(Warning, TEXT("Could not run Anim Node - No Anim Node Selected]")); 
 		return;
@@ -91,7 +91,7 @@ FRigUnit_RunAnimNode_Execute()
 	if (WorkData.AnimNodeInstance == nullptr)
 	{
 		// allocate AnimNode instance
-		WorkData.AnimNodeInstance = AnimNode.Get()->CreateInstance(Context);
+		WorkData.AnimNodeInstance = RootAnimNode.Get()->CreateInstance(Context);
 	}
 
 	TRACE_UAF_MODULE(ModuleInstance);
