AnimOps: Add missing default values on FUAFSteeringAnimOp

#Synced-CL 52381678

diff --git a/Engine/Plugins/Experimental/UAF/UAFAnimNode/Source/UAFAnimNode/Public/UAF/AnimOps/UAFSteeringAnimOp.h b/Engine/Plugins/Experimental/UAF/UAFAnimNode/Source/UAFAnimNode/Public/UAF/AnimOps/UAFSteeringAnimOp.h
--- a/Engine/Plugins/Experimental/UAF/UAFAnimNode/Source/UAFAnimNode/Public/UAF/AnimOps/UAFSteeringAnimOp.h
+++ b/Engine/Plugins/Experimental/UAF/UAFAnimNode/Source/UAFAnimNode/Public/UAF/AnimOps/UAFSteeringAnimOp.h
@@ -39,10 +39,10 @@ struct FUAFSteeringAnimOp : public FUAFAnimOp
 	FVector AngularVelocity = FVector::ZeroVector;
 
 	UPROPERTY()
-	FQuat TargetOrientation;
+	FQuat TargetOrientation = FQuat::Identity;
 
 	UPROPERTY()
-	FTransform RootBoneTransform;
+	FTransform RootBoneTransform = FTransform::Identity;
 
 	UPROPERTY()
 	float ProceduralTargetTime = 0.2f;
