Add AnimOp tuple and array view support

#Synced-CL 52402383

diff --git a/Engine/Plugins/Experimental/UAF/UAFAnimNode/Source/UAFAnimNode/Private/UAF/AnimOpCore/UAFAnimOpArrayView.cpp b/Engine/Plugins/Experimental/UAF/UAFAnimNode/Source/UAFAnimNode/Private/UAF/AnimOpCore/UAFAnimOpArrayView.cpp
new file mode 100644
--- /dev/null
+++ b/Engine/Plugins/Experimental/UAF/UAFAnimNode/Source/UAFAnimNode/Private/UAF/AnimOpCore/UAFAnimOpArrayView.cpp
@@ -0,0 +1,54 @@
+// Copyright Epic Games, Inc. All Rights Reserved.
+
+#include "UAF/AnimOpCore/UAFAnimOpArrayView.h"
+
+#include "UAF/AnimOpCore/UAFAnimOpValueEvaluator.h"
+
+namespace UE::UAF
+{
+	void FUAFAnimOpArrayView::GetAnimOps(TAdderRef<const FUAFAnimOp*> OutAnimOps) const
+	{
+		for (const FUAFAnimOp* AnimOp : AnimOps)
+		{
+			OutAnimOps.Add(AnimOp);
+		}
+	}
+
+	void FUAFAnimOpArrayView::EvaluateValues(FUAFAnimOpValueEvaluator& Evaluator)
+	{
+		TUAFStack<FValueBundleCoWRef>& EvaluationStack = Evaluator.GetEvaluationStack();
+
+		for (FUAFAnimOp* AnimOp : AnimOps)
+		{
+			if (EvaluationStack.Num() < AnimOp->GetNumInputs()) [[unlikely]]
+			{
+				UE_LOG(LogAnimation, Warning, TEXT("Too few inputs provided, AnimOp will be skipped."));
+				continue;
+			}
+
+			AnimOp->EvaluateValues(Evaluator);
+		}
+	}
+
+	void FUAFAnimOpArrayView::EvaluateNotifies(FUAFAnimOpNotifyEvaluator& Evaluator)
+	{
+		for (FUAFAnimOp* AnimOp : AnimOps)
+		{
+			if (!AnimOp->HasEvaluateNotifies()) [[likely]]
+			{
+				// We skip AnimOps that don't implement EvaluateNotifies since most of them don't
+				continue;
+			}
+
+			AnimOp->EvaluateNotifies(Evaluator);
+		}
+	}
+
+	void FUAFAnimOpArrayView::EvaluateSynchronization(FUAFAnimOpSyncEvaluator& Evaluator)
+	{
+		for (FUAFAnimOp* AnimOp : AnimOps)
+		{
+			AnimOp->EvaluateSynchronization(Evaluator);
+		}
+	}
+}
diff --git a/Engine/Plugins/Experimental/UAF/UAFAnimNode/Source/UAFAnimNode/Private/UAF/AnimOpCore/UAFInstancedAnimOpList.cpp b/Engine/Plugins/Experimental/UAF/UAFAnimNode/Source/UAFAnimNode/Private/UAF/AnimOpCore/UAFInstancedAnimOpList.cpp
--- a/Engine/Plugins/Experimental/UAF/UAFAnimNode/Source/UAFAnimNode/Private/UAF/AnimOpCore/UAFInstancedAnimOpList.cpp
+++ b/Engine/Plugins/Experimental/UAF/UAFAnimNode/Source/UAFAnimNode/Private/UAF/AnimOpCore/UAFInstancedAnimOpList.cpp
@@ -3,19 +3,39 @@
 #include "UAF/AnimOpCore/UAFInstancedAnimOpList.h"
 
 #include "UAF/AnimOpCore/UAFAnimOp.h"
+#include "UAF/AnimOpCore/UAFAnimOpListBase.h"
 
 #include UE_INLINE_GENERATED_CPP_BY_NAME(UAFInstancedAnimOpList)
 
 namespace UE::UAF
 {
-	FUAFInstancedAnimOpList::FUAFInstancedAnimOpList(const TArray<FUAFAnimOp*>& InAnimOps)
+	namespace Private
 	{
-		for (const FUAFAnimOp* AnimOp : InAnimOps)
+		static void CollectAnimOps(const TArrayView<const FUAFAnimOp* const>& AnimOps, FUAFInstancedAnimOpList& OutAnimOps)
 		{
-			const UScriptStruct* AnimOpType = AnimOp->GetStruct();
+			for (const FUAFAnimOp* AnimOp : AnimOps)
+			{
+				const UScriptStruct* AnimOpType = AnimOp->GetStruct();
+				if (AnimOpType == FUAFAnimOpListBase::StaticStruct())
+				{
+					const FUAFAnimOpListBase* AnimOpList = static_cast<const FUAFAnimOpListBase*>(AnimOp);
+
+					TArray<const FUAFAnimOp*> NestedAnimOps;
+					AnimOpList->GetAnimOps(NestedAnimOps);
 
-			AnimOps.Add(FInstancedStruct(AnimOpType));
-			AnimOpType->CopyScriptStruct(AnimOps.Last().GetMutableMemory(), AnimOp);
+					CollectAnimOps(NestedAnimOps, OutAnimOps);
+				}
+				else
+				{
+					OutAnimOps.AnimOps.Add(FInstancedStruct(AnimOpType));
+					AnimOpType->CopyScriptStruct(OutAnimOps.AnimOps.Last().GetMutableMemory(), AnimOp);
+				}
+			}
 		}
 	}
+
+	FUAFInstancedAnimOpList::FUAFInstancedAnimOpList(const TArray<FUAFAnimOp*>& InAnimOps)
+	{
+		Private::CollectAnimOps(MakeConstArrayView(InAnimOps.GetData(), InAnimOps.Num()), *this);
+	}
 }
diff --git a/Engine/Plugins/Experimental/UAF/UAFAnimNode/Source/UAFAnimNode/Public/UAF/AnimOpCore/UAFAnimOpArrayView.h b/Engine/Plugins/Experimental/UAF/UAFAnimNode/Source/UAFAnimNode/Public/UAF/AnimOpCore/UAFAnimOpArrayView.h
new file mode 100644
--- /dev/null
+++ b/Engine/Plugins/Experimental/UAF/UAFAnimNode/Source/UAFAnimNode/Public/UAF/AnimOpCore/UAFAnimOpArrayView.h
@@ -0,0 +1,67 @@
+// Copyright Epic Games, Inc. All Rights Reserved.
+
+#pragma once
+
+#include "CoreMinimal.h"
+#include "UAF/AnimOpCore/UAFAnimOpListBase.h"
+
+#define UE_API UAFANIMNODE_API
+
+namespace UE::UAF
+{
+	/**
+	 * FUAFAnimOpArrayView
+	 *
+	 * An array view of AnimOps.
+	 * AnimOps are evaluated in the ascending order.
+	 */
+	struct FUAFAnimOpArrayView : public FUAFAnimOpListBase
+	{
+		FUAFAnimOpArrayView();
+
+		// Creates an array view from a C-style array
+		template<size_t Count> static FUAFAnimOpArrayView Make(FUAFAnimOp*(&AnimOps)[Count]);
+
+		// Creates an array view from a TArray with any allocator
+		template<class AllocatorType> static FUAFAnimOpArrayView Make(const TArray<FUAFAnimOp*, AllocatorType>& AnimOps);
+
+		// FUAFAnimOp impl
+		UE_API virtual void EvaluateValues(FUAFAnimOpValueEvaluator& Evaluator) override;
+		UE_API virtual void EvaluateNotifies(FUAFAnimOpNotifyEvaluator& Evaluator) override;
+		UE_API virtual void EvaluateSynchronization(FUAFAnimOpSyncEvaluator& Evaluator) override;
+
+	protected:
+		// FUAFAnimOpListBase impl
+		UE_API virtual void GetAnimOps(TAdderRef<const FUAFAnimOp*> OutAnimOps) const override;
+
+		// An array view of AnimOps, evaluated in ascending order
+		TArrayView<FUAFAnimOp*> AnimOps;
+	};
+
+	//////////////////////////////////////////////////////////////////////////
+	// Inline implementation
+
+	inline FUAFAnimOpArrayView::FUAFAnimOpArrayView()
+		: FUAFAnimOpListBase()
+	{
+		InitializeAs<FUAFAnimOpArrayView>();
+	}
+
+	template<size_t Count>
+	inline FUAFAnimOpArrayView FUAFAnimOpArrayView::Make(FUAFAnimOp* (&AnimOps)[Count])
+	{
+		FUAFAnimOpArrayView Result;
+		Result.AnimOps = MakeArrayView(AnimOps, Count);
+		return Result;
+	}
+
+	template<class AllocatorType>
+	inline FUAFAnimOpArrayView FUAFAnimOpArrayView::Make(const TArray<FUAFAnimOp*, AllocatorType>& AnimOps)
+	{
+		FUAFAnimOpArrayView Result;
+		Result.AnimOps = MakeArrayView(AnimOps);
+		return Result;
+	}
+}
+
+#undef UE_API
diff --git a/Engine/Plugins/Experimental/UAF/UAFAnimNode/Source/UAFAnimNode/Public/UAF/AnimOpCore/UAFAnimOpListBase.h b/Engine/Plugins/Experimental/UAF/UAFAnimNode/Source/UAFAnimNode/Public/UAF/AnimOpCore/UAFAnimOpListBase.h
new file mode 100644
--- /dev/null
+++ b/Engine/Plugins/Experimental/UAF/UAFAnimNode/Source/UAFAnimNode/Public/UAF/AnimOpCore/UAFAnimOpListBase.h
@@ -0,0 +1,56 @@
+// Copyright Epic Games, Inc. All Rights Reserved.
+
+#pragma once
+
+#include "CoreMinimal.h"
+#include "Containers/AdderRef.h"
+#include "UAF/AnimOpCore/UAFAnimOp.h"
+
+#include "UAFAnimOpListBase.generated.h"
+
+namespace UE::UAF
+{
+	/**
+	 * FUAFAnimOpListBase
+	 *
+	 * Base type for lists of AnimOps.
+	 */
+	USTRUCT()
+	struct FUAFAnimOpListBase : public FUAFAnimOp
+	{
+		GENERATED_BODY()
+		UAF_DECLARE_ANIMOP(FUAFAnimOpListBase)
+
+		// ForceInit constructor for TCppStructOps that creates an invalid AnimOp
+		explicit FUAFAnimOpListBase(EForceInit);
+
+		// Returns the list of AnimOps contained within
+		virtual void GetAnimOps(TAdderRef<const FUAFAnimOp*> OutAnimOps) const PURE_VIRTUAL(FUAFAnimOpListBase::GetAnimOps,);
+
+	protected:
+		FUAFAnimOpListBase();
+	};
+
+	//////////////////////////////////////////////////////////////////////////
+	// Inline implementation
+
+	inline FUAFAnimOpListBase::FUAFAnimOpListBase(EForceInit ForceInit)
+		: FUAFAnimOp(ForceInit)
+	{
+	}
+
+	inline FUAFAnimOpListBase::FUAFAnimOpListBase()
+		: FUAFAnimOp(0)	// Lists have 0 inputs since they are just acting as a proxy
+	{
+	}
+}
+
+template<>
+struct TStructOpsTypeTraits<UE::UAF::FUAFAnimOpListBase> : TStructOpsTypeTraitsBase2<UE::UAF::FUAFAnimOpListBase>
+{
+	enum
+	{
+		WithNoInitConstructor = true,
+		WithPureVirtual = true,
+	};
+};
diff --git a/Engine/Plugins/Experimental/UAF/UAFAnimNode/Source/UAFAnimNode/Public/UAF/AnimOpCore/UAFAnimOpTuple.h b/Engine/Plugins/Experimental/UAF/UAFAnimNode/Source/UAFAnimNode/Public/UAF/AnimOpCore/UAFAnimOpTuple.h
new file mode 100644
--- /dev/null
+++ b/Engine/Plugins/Experimental/UAF/UAFAnimNode/Source/UAFAnimNode/Public/UAF/AnimOpCore/UAFAnimOpTuple.h
@@ -0,0 +1,112 @@
+// Copyright Epic Games, Inc. All Rights Reserved.
+
+#pragma once
+
+#include "CoreMinimal.h"
+#include "UAF/AnimOpCore/UAFAnimOpListBase.h"
+#include "UAF/AnimOpCore/UAFAnimOpValueEvaluator.h"
+
+namespace UE::UAF
+{
+	/**
+	 * TUAFAnimOpTuple
+	 *
+	 * A tuple of AnimOps.
+	 * AnimOps are evaluated in the order specified: left to right.
+	 */
+	template<typename... AnimOpTypes>
+	struct TUAFAnimOpTuple : public FUAFAnimOpListBase
+	{
+		TUAFAnimOpTuple();
+
+		// Returns the AnimOp with the specified index from the specified order
+		template<uint32 Index> UE_REWRITE decltype(auto) Get();
+		template<uint32 Index> UE_REWRITE decltype(auto) Get() const;
+
+		// FUAFAnimOp impl
+		virtual void EvaluateValues(FUAFAnimOpValueEvaluator& Evaluator) override;
+		virtual void EvaluateNotifies(FUAFAnimOpNotifyEvaluator& Evaluator) override;
+		virtual void EvaluateSynchronization(FUAFAnimOpSyncEvaluator& Evaluator) override;
+
+	protected:
+		// FUAFAnimOpListBase impl
+		virtual void GetAnimOps(TAdderRef<const FUAFAnimOp*> OutAnimOps) const override;
+
+		// A list of AnimOps, evaluated left to right
+		TTuple<AnimOpTypes...> AnimOps;
+	};
+
+	//////////////////////////////////////////////////////////////////////////
+	// Inline implementation
+
+	template<typename... AnimOpTypes>
+	inline TUAFAnimOpTuple<AnimOpTypes...>::TUAFAnimOpTuple()
+		: FUAFAnimOpListBase()
+	{
+		InitializeAs<TUAFAnimOpTuple<AnimOpTypes...>>();
+	}
+
+	template<typename... AnimOpTypes>
+	template<uint32 Index>
+	inline decltype(auto) TUAFAnimOpTuple<AnimOpTypes...>::Get()
+	{
+		return AnimOps.template Get<Index>();
+	}
+
+	template<typename... AnimOpTypes>
+	template<uint32 Index>
+	inline decltype(auto) TUAFAnimOpTuple<AnimOpTypes...>::Get() const
+	{
+		return AnimOps.template Get<Index>();
+	}
+
+	template<typename... AnimOpTypes>
+	inline void TUAFAnimOpTuple<AnimOpTypes...>::GetAnimOps(TAdderRef<const FUAFAnimOp*> OutAnimOps) const
+	{
+		VisitTupleElements([&OutAnimOps](auto& AnimOp) { OutAnimOps.Add(&AnimOp); }, AnimOps);
+	}
+
+	template<typename... AnimOpTypes>
+	inline void TUAFAnimOpTuple<AnimOpTypes...>::EvaluateValues(FUAFAnimOpValueEvaluator& Evaluator)
+	{
+		TUAFStack<FValueBundleCoWRef>& EvaluationStack = Evaluator.GetEvaluationStack();
+
+		VisitTupleElements(
+			[&Evaluator, &EvaluationStack](auto& AnimOp)
+			{
+				if (EvaluationStack.Num() < AnimOp.GetNumInputs()) [[unlikely]]
+				{
+					UE_LOG(LogAnimation, Warning, TEXT("Too few inputs provided, AnimOp will be skipped."));
+					return;
+				}
+
+				AnimOp.EvaluateValues(Evaluator);
+			}, AnimOps);
+	}
+
+	template<typename... AnimOpTypes>
+	inline void TUAFAnimOpTuple<AnimOpTypes...>::EvaluateNotifies(FUAFAnimOpNotifyEvaluator& Evaluator)
+	{
+		VisitTupleElements(
+			[&Evaluator](auto& AnimOp)
+			{
+				if (!AnimOp.HasEvaluateNotifies()) [[likely]]
+				{
+					// We skip AnimOps that don't implement EvaluateNotifies since most of them don't
+					return;
+				}
+
+				AnimOp.EvaluateNotifies(Evaluator);
+			}, AnimOps);
+	}
+
+	template<typename... AnimOpTypes>
+	inline void TUAFAnimOpTuple<AnimOpTypes...>::EvaluateSynchronization(FUAFAnimOpSyncEvaluator& Evaluator)
+	{
+		VisitTupleElements(
+			[&Evaluator](auto& AnimOp)
+			{
+				AnimOp.EvaluateSynchronization(Evaluator);
+			}, AnimOps);
+	}
+}
