UAF AnimNode: Add Make Dynamic Additive node

#Synced-CL 52496234

diff --git a/Engine/Plugins/Experimental/UAF/UAFAnimNode/Source/UAFAnimNode/Private/UAF/AnimNodes/UAFMakeDynamicAdditive.cpp b/Engine/Plugins/Experimental/UAF/UAFAnimNode/Source/UAFAnimNode/Private/UAF/AnimNodes/UAFMakeDynamicAdditive.cpp
new file mode 100644
--- /dev/null
+++ b/Engine/Plugins/Experimental/UAF/UAFAnimNode/Source/UAFAnimNode/Private/UAF/AnimNodes/UAFMakeDynamicAdditive.cpp
@@ -0,0 +1,101 @@
+// Copyright Epic Games, Inc. All Rights Reserved.
+
+
+#include "UAF/AnimNodes/UAFMakeDynamicAdditive.h"
+
+#include "Logging/StructuredLog.h"
+#include "UAF/AnimOps/UAFNullAnimOp.h"
+
+
+namespace UE::UAF
+{
+
+	FUAFAnimNodePtr FUAFMakeDynamicAdditiveData::CreateInstance(FUAFAnimGraphUpdateContext& Context) const
+	{
+		return MakeAnimNode<FUAFMakeDynamicAdditiveNode>(Context, *this);
+	}
+
+	FUAFMakeDynamicAdditiveNode::FUAFMakeDynamicAdditiveNode(FUAFAnimGraphUpdateContext& Context, const FUAFMakeDynamicAdditiveData& InData)
+		: FUAFAnimNode(Context)
+		, Data(&InData)
+	{
+		InitializeAs<FUAFMakeDynamicAdditiveNode>(Context);
+
+		// Create both child nodes up front,
+		// Base and Source are both required.
+		FUAFAnimNodePtr BaseChild;
+		FUAFAnimNodePtr SourceChild;
+
+		if (Data->BaseNode.IsValid())
+		{
+			BaseChild = Data->BaseNode.Get().CreateInstance(Context);
+		}
+
+		if (Data->SourceNode.IsValid())
+		{
+			SourceChild = Data->SourceNode.Get().CreateInstance(Context);
+		}
+
+		if (BaseChild && SourceChild)
+		{
+			AddChild(BaseChild);
+			AddChild(SourceChild);
+
+			MakeDynamicAdditiveAnimOp.SetMeshSpaceAdditive(Data->bMeshSpaceAdditive);
+			SetPostAnimOp(&MakeDynamicAdditiveAnimOp);
+		}
+		else
+		{
+			SetPostAnimOp(FUAFNullAnimOp::Get());
+			UE_LOGFMT(LogAnimation, Error, "FUAFMakeDynamicAdditiveNode - Both Base and Source inputs are required. Host Object: [{0}]", GetNameSafe(Context.GetHostObject()));
+		}
+	}
+
+	void FUAFMakeDynamicAdditiveNode::PreUpdate(FUAFAnimGraphUpdateContext& GraphContext)
+	{
+		FUAFAnimNode::PreUpdate(GraphContext);
+
+		check(Data);
+		// If I have any children, it means I have both children.
+		if (HasChildren())
+		{
+			// Both inputs have full weight.
+			GetChildAt(BaseChildIndex)->SetTotalWeight(GetTotalWeight());
+			GetChildAt(SourceChildIndex)->SetTotalWeight(GetTotalWeight());
+		}
+	}
+
+	void* FUAFMakeDynamicAdditiveNode::GetInterface(FUAFAnimNodeInterfaceId Id)
+	{
+		if (!HasChildren())
+		{
+			return nullptr;
+		}
+
+		if (const FUAFAnimNodePtr& Source = GetChildAt(SourceChildIndex))
+		{
+			return Source->GetInterface(Id);
+		}
+
+		return nullptr;
+	}
+
+	void FUAFMakeDynamicAdditiveNode::AddReferencedObjects(FUAFAnimNode* This, FReferenceCollector& Collector)
+	{
+		FUAFMakeDynamicAdditiveNode* ThisNode = static_cast<FUAFMakeDynamicAdditiveNode*>(This);
+		Collector.AddPropertyReferencesWithStructARO(FUAFMakeDynamicAdditiveAnimOp::StaticStruct(), &ThisNode->MakeDynamicAdditiveAnimOp);
+	}
+
+#if UAF_TRACE_ENABLED
+	FString FUAFMakeDynamicAdditiveNode::GetDebugName() const
+	{
+		static FString MakeDynamicAdditiveNodeName("Make Dynamic Additive");
+		return MakeDynamicAdditiveNodeName;
+	}
+
+	UStruct* FUAFMakeDynamicAdditiveNode::GetDebugStruct() const
+	{
+		return FUAFMakeDynamicAdditiveData::StaticStruct();
+	}
+#endif
+}
diff --git a/Engine/Plugins/Experimental/UAF/UAFAnimNode/Source/UAFAnimNode/Private/UAF/AnimOps/UAFMakeDynamicAdditiveAnimOp.cpp b/Engine/Plugins/Experimental/UAF/UAFAnimNode/Source/UAFAnimNode/Private/UAF/AnimOps/UAFMakeDynamicAdditiveAnimOp.cpp
new file mode 100644
--- /dev/null
+++ b/Engine/Plugins/Experimental/UAF/UAFAnimNode/Source/UAFAnimNode/Private/UAF/AnimOps/UAFMakeDynamicAdditiveAnimOp.cpp
@@ -0,0 +1,83 @@
+// Copyright Epic Games, Inc. All Rights Reserved.
+
+
+#include "UAF/AnimOps/UAFMakeDynamicAdditiveAnimOp.h"
+
+#include "UAF/AnimOpCore/UAFAnimOpValueEvaluator.h"
+#include "UAF/ValueRuntime/Transformers/AdditiveSpace.h"
+#include "UAF/ValueRuntime/Transformers/BoneSpace.h"
+#include "UAF/ValueRuntime/ValueBundle.h"
+#include "Logging/StructuredLog.h"
+
+namespace UE::UAF
+{
+	FUAFMakeDynamicAdditiveAnimOp::FUAFMakeDynamicAdditiveAnimOp()
+		: FUAFAnimOp(2)
+	{
+		InitializeAs<FUAFMakeDynamicAdditiveAnimOp>();
+	}
+
+	void FUAFMakeDynamicAdditiveAnimOp::EvaluateValues(FUAFAnimOpValueEvaluator& Evaluator)
+	{
+		// Stack order: Base (bottom), Source (top)
+		// Pop Source first (top of stack), then Base (underneath)
+		FValueBundleCoWRef SourceInput = Evaluator.GetEvaluationStack().Pop();
+		FValueBundleCoWRef BaseInput = Evaluator.GetEvaluationStack().Pop();
+
+		const FValueSpace AdditiveSpace = bMeshSpaceAdditive
+			? FValueSpace(EMixedSpaceFlags::MeshRotation, true)
+			: FValueSpace(EValueSpaceType::Local, true);
+
+		bool bInvalidSetup = false;
+		if (BaseInput.Get().IsEmpty())
+		{
+			bInvalidSetup = true;
+			UE_LOGF(LogAnimation, Error, "FUAFMakeDynamicAdditiveAnimOp - Base input is empty.");
+		}
+		else if (BaseInput.Get().IsAdditive())
+		{
+			bInvalidSetup = true;
+			UE_LOGF(LogAnimation, Error, "FUAFMakeDynamicAdditiveAnimOp - Base input is additive. Expected non-additive base.");
+		}
+
+		if (SourceInput.Get().IsEmpty())
+		{
+			bInvalidSetup = true;
+			UE_LOGF(LogAnimation, Error, "FUAFMakeDynamicAdditiveAnimOp - Source input is empty.");
+		}
+		else if (SourceInput.Get().IsAdditive())
+		{
+			bInvalidSetup = true;
+			UE_LOGF(LogAnimation, Error, "FUAFMakeDynamicAdditiveAnimOp - Source input is additive. Expected non-additive Source.");
+		}
+
+		if (bInvalidSetup)
+		{
+			FValueBundleStack Output(Evaluator.GetActiveNamedSet());
+			Output.InitWithValueSpace(AdditiveSpace);
+			Evaluator.GetEvaluationStack().Push(FValueBundleCoWRef::MakeFrom(MoveTemp(Output)));
+			return;
+		}
+
+		if (bMeshSpaceAdditive)
+		{
+			// Convert both inputs to mesh rotation space before computing the additive delta
+			BaseInput.ForceMutable();
+			SourceInput.ForceMutable();
+			Transformers::FBoneSpace::LocalToMeshRotation(BaseInput.GetMutable(), BaseInput.GetMutable());
+			Transformers::FBoneSpace::LocalToMeshRotation(SourceInput.GetMutable(), SourceInput.GetMutable());
+		}
+
+		FValueBundleStack Output(Evaluator.GetActiveNamedSet());
+		Output.InitWithValueSpace(AdditiveSpace);
+
+		Transformers::FMakeAdditiveSpace::Apply(Evaluator.GetTransformerMap(), BaseInput.Get(), SourceInput.Get(), Output);
+
+		Evaluator.GetEvaluationStack().Push(FValueBundleCoWRef::MakeFrom(MoveTemp(Output)));
+	}
+
+	void FUAFMakeDynamicAdditiveAnimOp::SetMeshSpaceAdditive(const bool bInMeshSpaceAdditive)
+	{
+		bMeshSpaceAdditive = bInMeshSpaceAdditive;
+	}
+}
diff --git a/Engine/Plugins/Experimental/UAF/UAFAnimNode/Source/UAFAnimNode/Public/UAF/AnimNodes/UAFMakeDynamicAdditive.h b/Engine/Plugins/Experimental/UAF/UAFAnimNode/Source/UAFAnimNode/Public/UAF/AnimNodes/UAFMakeDynamicAdditive.h
new file mode 100644
--- /dev/null
+++ b/Engine/Plugins/Experimental/UAF/UAFAnimNode/Source/UAFAnimNode/Public/UAF/AnimNodes/UAFMakeDynamicAdditive.h
@@ -0,0 +1,63 @@
+// Copyright Epic Games, Inc. All Rights Reserved.
+
+#pragma once
+
+#include "CoreMinimal.h"
+#include "BindableValue/UAFBindableTypes.h"
+#include "UAF/AnimNodeCore/UAFAnimNodeData.h"
+#include "UAF/AnimOps/UAFMakeDynamicAdditiveAnimOp.h"
+
+#include "UAFMakeDynamicAdditive.generated.h"
+
+#define UE_API UAFANIMNODE_API
+
+namespace UE::UAF
+{
+	// Node to generate an additive animation from Base and Source poses.
+	// The output is an additive pose. Additive =  (Source - Base).
+	// Generates an additive identity in case of a bad content setup.
+	USTRUCT(DisplayName = "Make Dynamic Additive")
+	struct FUAFMakeDynamicAdditiveData : public FUAFAnimNodeData
+	{
+		GENERATED_BODY()
+
+		// Node providing the base (reference) animation
+		UPROPERTY(EditAnywhere, Category = "NodeData")
+		TInstancedStruct<FUAFAnimNodeData> BaseNode;
+
+		// Node providing the Source animation to turn into an additive
+		UPROPERTY(EditAnywhere, Category = "NodeData")
+		TInstancedStruct<FUAFAnimNodeData> SourceNode;
+
+		// Whether to compute the additive delta in mesh space (rotation only).
+		UPROPERTY(EditAnywhere, Category = "NodeData")
+		bool bMeshSpaceAdditive = false;
+
+		UE_API virtual FUAFAnimNodePtr CreateInstance(FUAFAnimGraphUpdateContext& Context) const override;
+	};
+
+	class FUAFMakeDynamicAdditiveNode : public FUAFAnimNode
+	{
+	public:
+		UE_API explicit FUAFMakeDynamicAdditiveNode(FUAFAnimGraphUpdateContext& Context, const FUAFMakeDynamicAdditiveData& InData);
+
+		// FUAFAnimNode impl
+		UE_API virtual void PreUpdate(FUAFAnimGraphUpdateContext& GraphContext) override;
+		UE_API virtual void* GetInterface(FUAFAnimNodeInterfaceId Id) override;
+		UE_API static void AddReferencedObjects(FUAFAnimNode* This, FReferenceCollector& Collector);
+
+#if UAF_TRACE_ENABLED
+		UE_API virtual FString GetDebugName() const override;
+		UE_API virtual UStruct* GetDebugStruct() const override;
+#endif
+
+	private:
+		static constexpr int32 BaseChildIndex = 0;
+		static constexpr int32 SourceChildIndex = 1;
+
+		const FUAFMakeDynamicAdditiveData* Data = nullptr;
+		FUAFMakeDynamicAdditiveAnimOp MakeDynamicAdditiveAnimOp;
+	};
+}
+
+#undef UE_API
diff --git a/Engine/Plugins/Experimental/UAF/UAFAnimNode/Source/UAFAnimNode/Public/UAF/AnimOps/UAFMakeDynamicAdditiveAnimOp.h b/Engine/Plugins/Experimental/UAF/UAFAnimNode/Source/UAFAnimNode/Public/UAF/AnimOps/UAFMakeDynamicAdditiveAnimOp.h
new file mode 100644
--- /dev/null
+++ b/Engine/Plugins/Experimental/UAF/UAFAnimNode/Source/UAFAnimNode/Public/UAF/AnimOps/UAFMakeDynamicAdditiveAnimOp.h
@@ -0,0 +1,40 @@
+// Copyright Epic Games, Inc. All Rights Reserved.
+
+#pragma once
+
+#include "CoreMinimal.h"
+#include "UAF/AnimOpCore/UAFAnimOp.h"
+
+#include "UAFMakeDynamicAdditiveAnimOp.generated.h"
+
+#define UE_API UAFANIMNODE_API
+
+namespace UE::UAF
+{
+	class FUAFAnimOpValueEvaluator;
+
+	// Anim Op to generate an additive animation from Base and Source poses.
+	// The output is (Source - Base) from the stack.
+	USTRUCT()
+	struct FUAFMakeDynamicAdditiveAnimOp : public FUAFAnimOp
+	{
+		GENERATED_BODY()
+		UAF_DECLARE_ANIMOP(FUAFMakeDynamicAdditiveAnimOp)
+
+	public:
+		FUAFMakeDynamicAdditiveAnimOp();
+
+		// FUAFAnimOp impl
+		UE_API virtual void EvaluateValues(FUAFAnimOpValueEvaluator& Evaluator) override;
+
+		UE_API void SetMeshSpaceAdditive(bool bInMeshSpaceAdditive);
+
+	private:
+
+		// Whether to compute the additive delta in mesh space (rotation only).
+		UPROPERTY(VisibleAnywhere, Category = Properties)
+		bool bMeshSpaceAdditive = false;
+	};
+}
+
+#undef UE_API
