Add missing export for destructor

#Synced-CL 52402383

diff --git a/Engine/Plugins/Experimental/UAF/UAFAnimNode/Source/UAFAnimNode/Public/UAF/AnimNodeCore/UAFAnimNodeReferenceCollector.h b/Engine/Plugins/Experimental/UAF/UAFAnimNode/Source/UAFAnimNode/Public/UAF/AnimNodeCore/UAFAnimNodeReferenceCollector.h
--- a/Engine/Plugins/Experimental/UAF/UAFAnimNode/Source/UAFAnimNode/Public/UAF/AnimNodeCore/UAFAnimNodeReferenceCollector.h
+++ b/Engine/Plugins/Experimental/UAF/UAFAnimNode/Source/UAFAnimNode/Public/UAF/AnimNodeCore/UAFAnimNodeReferenceCollector.h
@@ -21,7 +21,7 @@ namespace UE::UAF
 		using AddReferencedObjectsFn = void(*)(FUAFAnimNode*, FReferenceCollector&);
 
 		FUAFAnimNodeReferenceCollector() = default;
-		~FUAFAnimNodeReferenceCollector();
+		UE_API ~FUAFAnimNodeReferenceCollector();
 
 		// Registers the provided anim node with the GC system
 		// Once registered, TUAFAnimNode<T>::AddReferencedObjects will be called on it during GC
