NullPose AnimNode
- a dummy anim node which you can create in place of a valid child node when data is not set up correctly.  Simplifies special case handling in nodes, because you can guarantee your child is non-null, and don't have to have special cases to change the Op to a NullAnimOp

#Synced-CL 52512946

diff --git a/Engine/Plugins/Experimental/UAF/UAFAnimNode/Source/UAFAnimNode/Private/UAF/AnimNodes/UAFNullPoseNode.cpp b/Engine/Plugins/Experimental/UAF/UAFAnimNode/Source/UAFAnimNode/Private/UAF/AnimNodes/UAFNullPoseNode.cpp
new file mode 100644
--- /dev/null
+++ b/Engine/Plugins/Experimental/UAF/UAFAnimNode/Source/UAFAnimNode/Private/UAF/AnimNodes/UAFNullPoseNode.cpp
@@ -0,0 +1,27 @@
+// Copyright Epic Games, Inc. All Rights Reserved.
+
+#include "UAF/AnimNodes/UAFNullPoseNode.h"
+#include "UAF/AnimNodeCore/UAFAnimNodeUpdate.h"
+#include "UAF/AnimNodeCore/UAFAnimNodeData.h"
+
+namespace UE::UAF
+{
+	FUAFNullPoseNode::FUAFNullPoseNode(FUAFAnimGraphUpdateContext& Context)
+		: FUAFAnimNode(Context)
+	{
+		InitializeAs<FUAFNullPoseNode>(Context);
+		SetPostAnimOp(FUAFNullAnimOp::Get());
+	}
+
+#if UAF_TRACE_ENABLED
+	FString FUAFNullPoseNode::GetDebugName() const
+	{
+		return "Null AnimNode";
+	}
+
+	UStruct* FUAFNullPoseNode::GetDebugStruct() const
+	{
+		return FUAFAnimNodeData::StaticStruct();
+	}
+#endif
+}
diff --git a/Engine/Plugins/Experimental/UAF/UAFAnimNode/Source/UAFAnimNode/Public/UAF/AnimNodes/UAFNullPoseNode.h b/Engine/Plugins/Experimental/UAF/UAFAnimNode/Source/UAFAnimNode/Public/UAF/AnimNodes/UAFNullPoseNode.h
new file mode 100644
--- /dev/null
+++ b/Engine/Plugins/Experimental/UAF/UAFAnimNode/Source/UAFAnimNode/Public/UAF/AnimNodes/UAFNullPoseNode.h
@@ -0,0 +1,33 @@
+// Copyright Epic Games, Inc. All Rights Reserved.
+
+#pragma once
+
+#include "CoreMinimal.h"
+#include "UAF/AnimNodeCore/UAFAnimNode.h"
+#include "UAF/AnimOps/UAFNullAnimOp.h"
+
+#define UE_API UAFANIMNODE_API
+
+namespace UE::UAF
+{
+	// Dummy Anim node, produces a T Pose
+	// For use by other nodes when their child data is not set or for other failure cases.
+
+	class FUAFNullPoseNode : public FUAFAnimNode
+	{
+	public:
+		UE_API explicit FUAFNullPoseNode(FUAFAnimGraphUpdateContext& Context);
+
+		virtual void* GetInterface(FUAFAnimNodeInterfaceId Id)
+		{
+			return nullptr;
+		};
+
+#if UAF_TRACE_ENABLED
+		UE_API virtual FString GetDebugName() const override;
+		UE_API virtual UStruct* GetDebugStruct() const override;
+#endif
+	};
+}
+
+#undef UE_API
