UAF AnimNode: Add Apply Additive Node

#Synced-CL 52535972

diff --git a/Engine/Plugins/Experimental/UAF/UAFAnimNode/Source/UAFAnimNode/Private/UAF/AnimNodes/UAFApplyAdditive.cpp b/Engine/Plugins/Experimental/UAF/UAFAnimNode/Source/UAFAnimNode/Private/UAF/AnimNodes/UAFApplyAdditive.cpp
new file mode 100644
--- /dev/null
+++ b/Engine/Plugins/Experimental/UAF/UAFAnimNode/Source/UAFAnimNode/Private/UAF/AnimNodes/UAFApplyAdditive.cpp
@@ -0,0 +1,127 @@
+// Copyright Epic Games, Inc. All Rights Reserved.
+
+
+#include "UAF/AnimNodes/UAFApplyAdditive.h"
+
+#include "Logging/StructuredLog.h"
+#include "UAF/AnimOps/UAFNullAnimOp.h"
+
+
+namespace UE::UAF
+{
+	FUAFAnimNodePtr FUAFApplyAdditiveData::CreateInstance(FUAFAnimGraphUpdateContext& Context) const
+	{
+		return MakeAnimNode<FUAFApplyAdditiveNode>(Context, *this);
+	}
+
+	FUAFApplyAdditiveNode::FUAFApplyAdditiveNode(FUAFAnimGraphUpdateContext& Context, const FUAFApplyAdditiveData& InData)
+		: FUAFAnimNode(Context)
+		, Data(InData)
+	{
+		InitializeAs<FUAFApplyAdditiveNode>(Context);
+		
+		AddChild(nullptr);
+		AddChild(nullptr);
+		
+		// Create base and additive child nodes 
+		if (Data.BaseNode.IsValid())
+		{
+			if (const FUAFAnimNodePtr BaseChild = Data.BaseNode.Get().CreateInstance(Context))
+			{
+				SetChildAt(BaseChildIndex, BaseChild);
+
+				if (Data.AdditiveNode.IsValid())
+				{
+					if (const FUAFAnimNodePtr AdditiveChild = Data.AdditiveNode.Get().CreateInstance(Context))
+					{
+						SetChildAt(AdditiveChildIndex, AdditiveChild);
+					
+						// If both child nodes have been created successfully we create the additive anim op
+						const float ActiveAdditiveWeight = Data.AdditiveWeight.GetValue(Context.GetVariablesOwner());
+			
+						AdditiveAnimOp.SetAdditiveWeight(ActiveAdditiveWeight);
+						AdditiveAnimOp.SetBlendMask(Data.BlendMask);
+						SetPostAnimOp(&AdditiveAnimOp);
+					}
+				}
+			}
+		}
+			
+		if (GetChildAt(BaseChildIndex) == nullptr)
+		{
+			// Without base we provide a null op instead
+			SetPostAnimOp(FUAFNullAnimOp::Get());
+			UE_LOGFMT(LogAnimation, Error, "FUAFApplyAdditiveNode - No valid base input provided. Host Object: [{0}]", GetNameSafe(Context.GetHostObject()));
+		}
+	
+	}
+	
+	void FUAFApplyAdditiveNode::PreUpdate(FUAFAnimGraphUpdateContext& GraphContext)
+	{
+		FUAFAnimNode::PreUpdate(GraphContext);
+
+		if (const FUAFAnimNodePtr Base = GetChildAt(BaseChildIndex))
+		{
+			Base->SetTotalWeight(GetTotalWeight());
+			
+			if (const FUAFAnimNodePtr Additive = GetChildAt(AdditiveChildIndex))
+			{
+				const float ActiveAdditiveWeight = Data.AdditiveWeight.GetValue(GraphContext.GetVariablesOwner());
+				Additive->SetTotalWeight(GetTotalWeight() * ActiveAdditiveWeight);
+				
+				// Update relevant values for our AnimOp
+				AdditiveAnimOp.SetAdditiveWeight(ActiveAdditiveWeight);
+				AdditiveAnimOp.SetBlendMask(Data.BlendMask);
+			}
+		}
+	}
+
+	void* FUAFApplyAdditiveNode::GetInterface(FUAFAnimNodeInterfaceId Id)
+	{
+		// For now forward any interface calls to the base first and if invalid then to the additive 
+		if (const FUAFAnimNodePtr Base = GetChildAt(BaseChildIndex))
+		{
+			void* BaseInterface = Base->GetInterface(Id);
+			if (BaseInterface != nullptr)
+			{
+				return BaseInterface;
+			}
+		}
+		
+		if (const FUAFAnimNodePtr Additive = GetChildAt(AdditiveChildIndex))
+		{
+			void* AdditiveInterface = Additive->GetInterface(Id);
+			if (AdditiveInterface != nullptr)
+			{
+				return AdditiveInterface;
+			}
+		}
+		
+		return nullptr;
+	}
+
+	void FUAFApplyAdditiveNode::AddReferencedObjects(FUAFAnimNode* This, FReferenceCollector& Collector)
+	{
+		FUAFApplyAdditiveNode* ThisAdditiveNode = static_cast<FUAFApplyAdditiveNode*>(This);
+		if (ThisAdditiveNode)
+		{
+			FUAFApplyAdditiveData& ThisData = const_cast<FUAFApplyAdditiveData&>(ThisAdditiveNode->Data);
+			Collector.AddReferencedObject(ThisData.BlendMask);
+			Collector.AddPropertyReferencesWithStructARO(FUAFApplyAdditiveAnimOp::StaticStruct(), &ThisAdditiveNode->AdditiveAnimOp);
+		}
+	}
+	
+#if UAF_TRACE_ENABLED
+	FString FUAFApplyAdditiveNode::GetDebugName() const
+	{
+		static FString ApplyAdditiveNodeName("Apply Additive");
+		return ApplyAdditiveNodeName;
+	}
+
+	UStruct* FUAFApplyAdditiveNode::GetDebugStruct() const
+	{
+		return FUAFApplyAdditiveData::StaticStruct();
+	}
+#endif
+	
+}
diff --git a/Engine/Plugins/Experimental/UAF/UAFAnimNode/Source/UAFAnimNode/Private/UAF/AnimOps/UAFApplyAdditiveAnimOp.cpp b/Engine/Plugins/Experimental/UAF/UAFAnimNode/Source/UAFAnimNode/Private/UAF/AnimOps/UAFApplyAdditiveAnimOp.cpp
new file mode 100644
--- /dev/null
+++ b/Engine/Plugins/Experimental/UAF/UAFAnimNode/Source/UAFAnimNode/Private/UAF/AnimOps/UAFApplyAdditiveAnimOp.cpp
@@ -0,0 +1,197 @@
+// Copyright Epic Games, Inc. All Rights Reserved.
+
+
+#include "UAF/AnimOps/UAFApplyAdditiveAnimOp.h"
+
+#include "Animation/SkeletonRemappingRegistry.h"
+#include "UAF/AnimOpCore/UAFAnimOpValueEvaluator.h"
+#include "UAF/ValueRuntime/Transformers/AdditiveSpace.h"
+#include "UAF/ValueRuntime/Transformers/BoneSpace.h"
+#include "UAF/ValueRuntime/Transformers/Sanitize.h"
+
+namespace UE::UAF
+{
+	FUAFApplyAdditiveAnimOp::FUAFApplyAdditiveAnimOp()
+		: FUAFAnimOp(2) 
+	{
+		InitializeAs<FUAFApplyAdditiveAnimOp>();
+	}
+
+	void FUAFApplyAdditiveAnimOp::EvaluateValues(FUAFAnimOpValueEvaluator& Evaluator)
+	{
+		const FUAFAnimOpValueEvaluationContext& AnimOpEvaluationContext = Evaluator.GetActiveEvaluationContext();
+
+		// Pop our inputs - the additive being on top and base the next underneath 
+		FValueBundleCoWRef AdditiveInput = Evaluator.GetEvaluationStack().Pop();
+		FValueBundleCoWRef BaseInput = Evaluator.GetEvaluationStack().Pop();
+
+		const FValueBundle& Additive = AdditiveInput.Get();
+		FValueBundle& Base = BaseInput.GetMutable();
+		
+		// If the additive is not relevant or empty, we simply fall back to the base pose
+		// Both scenarios are valid and do not necessarily indicate a setup error so we will not raise a waring 
+		if (!FAnimWeight::IsRelevant(AdditiveWeight) || Additive.IsEmpty())
+		{
+			Evaluator.GetEvaluationStack().Push(MoveTemp(BaseInput));
+			return;
+		}
+
+		// If the additive animation is not empty, but it also is not an additive warn the user and fall back to the base pose instead
+		if (!Additive.IsAdditive())
+		{
+			Evaluator.GetEvaluationStack().Push(MoveTemp(BaseInput));
+			UE_LOG(LogAnimation, Error, TEXT("FUAFApplyAdditiveAnimOp::EvaluateValues - The provided additive animation is not an additive"));
+			return;
+		}
+		
+		// If the base is empty init it with default values 
+		if (Base.IsEmpty())
+		{
+			check(BaseInput.IsMutable());
+			BaseInput.GetMutable().InitWithValueSpace(FValueSpace(EValueSpaceType::Local));
+			
+			UE_LOG(LogAnimation, Warning, TEXT("FUAFApplyAdditiveAnimOp::EvaluateValues - The base input is empty. Will use a ref pose instead."))
+		}
+		
+		// Named sets must match as it ensures that our inputs have the same sizes/shapes
+		check(Base.GetNamedSet() == Additive.GetNamedSet());
+		
+		// If our additive input is in mesh rotation space, we ensure our base is in that space as well
+		bool bHasToConvertBaseToMeshSpace = false;
+		if (Additive.GetValueSpace().GetMixedSpaceFlags() == EMixedSpaceFlags::MeshRotation)
+		{
+			const FValueSpace BaseSpace = Base.GetValueSpace();
+			bHasToConvertBaseToMeshSpace = BaseSpace.GetMixedSpaceFlags() != EMixedSpaceFlags::MeshRotation;
+		}
+
+		// Modify our base input in-place
+		if (bHasToConvertBaseToMeshSpace)
+		{
+			Transformers::FBoneSpace::LocalToMeshRotation(Base, Base);
+		}
+
+		if (OptionalBlendMask)
+		{
+			FValueBundleStack PerValueWeights(Evaluator.GetActiveNamedSet());
+			BuildPerValueWeights(Evaluator, Base, PerValueWeights);
+			Transformers::FApplyAdditiveSpace::Apply(Evaluator.GetTransformerMap(), Base, Additive, PerValueWeights, AdditiveWeight, Base);
+		}
+		else
+		{
+			Transformers::FApplyAdditiveSpace::Apply(Evaluator.GetTransformerMap(), Base, Additive, AdditiveWeight, Base);
+		}
+		
+		if (bHasToConvertBaseToMeshSpace)
+		{
+			Transformers::FBoneSpace::MeshRotationToLocal(Base, Base);
+		}
+
+		Transformers::FSanitize::Apply(Evaluator.GetTransformerMap(), Base);
+		Evaluator.GetEvaluationStack().Push(MoveTemp(BaseInput));
+	}
+
+	void FUAFApplyAdditiveAnimOp::SetAdditiveWeight(const float InWeight)
+	{
+		AdditiveWeight = InWeight;
+	}
+
+	void FUAFApplyAdditiveAnimOp::SetBlendMask(const TObjectPtr<UUAFBlendMask> InBlendMask)
+	{
+		OptionalBlendMask = InBlendMask;
+	}
+
+	void FUAFApplyAdditiveAnimOp::BuildPerValueWeights(const FUAFAnimOpValueEvaluator& Evaluator, const FValueBundle& BaseInput, FValueBundleStack& OutPerValueWeights)
+	{
+		if (OptionalBlendMask == nullptr)
+		{
+			return;
+		}
+		
+		UScriptStruct* FloatAttributeType = FFloatAnimationAttribute::StaticStruct();
+		const TObjectPtr<USkeleton> SourceSkeleton = OptionalBlendMask->GetSkeleton();
+		const UUAFBlendMask::FSkeletonBoneWeightArray& BoneMaskWeights = OptionalBlendMask->GetSkeletonBoneWeights(); 
+		const UUAFBlendMask::FSkeletonCurveWeightArray& CurveMaskWeights = OptionalBlendMask->GetSkeletonCurveWeights(); 
+		const UUAFBlendMask::FSkeletonAttributeWeightArray& AttributeMaskWeights = OptionalBlendMask->GetSkeletonAttributeWeights();
+		
+		// Build our per-value weights (only for bound values)
+		FBoundMapCollection& PerValueWeightMaps = OutPerValueWeights.GetBoundValueMaps();
+		for (auto It = BaseInput.GetBoundValueMaps().CreateConstIterator(); It; ++It)
+		{
+			UScriptStruct* AttributeType = It.GetAttributeType();
+			// Every attribute type maps to a float - the weight we want to apply per attribute
+			FAttributeMappingKey MappingKey = FAttributeMappingKey::MakeFromTo(AttributeType, FloatAttributeType);
+			// Add the float attribute to our per value weight map
+			TBoundValueMap<FFloatAnimationAttribute>* SetWeights = PerValueWeightMaps.Add<FFloatAnimationAttribute>(MappingKey);
+			
+			// Bone transforms weights
+			if (const TBoundValueMap<FBoneTransformAnimationAttribute>* BoneMap = Cast<FBoneTransformAnimationAttribute>(It.GetMap()))
+			{
+				const FAttributeTypedSetPtr& BoneSet = BoneMap->GetTypedSet();
+				const int32 NumBonesInSet = BoneMap->Num();
+
+				check(OptionalBlendMask->GetSkeletonBoneWeights().Num() >= NumBonesInSet);
+				
+				const USkeleton* TargetSkeleton = Evaluator.GetActiveEvaluationContext().GetSkeleton();
+				const FSkeletonRemapping& SkeletonRemapping = UE::Anim::FSkeletonRemappingRegistry::Get().GetRemapping(SourceSkeleton, TargetSkeleton);
+				
+				for (FAttributeSetIndex SetIndex(0); SetIndex < NumBonesInSet; ++SetIndex)
+				{
+					const FAttributeBindingIndex BindingIndex = BoneSet->GetBindingIndex(SetIndex);
+					const FSkeletonPoseBoneIndex TargetSkeletonBoneIndex = BindingIndex;
+					int32 WeightIndex = TargetSkeletonBoneIndex.GetInt();
+					if (SkeletonRemapping.IsValid())
+					{
+						const FSkeletonPoseBoneIndex SourceSkeletonBoneIndex(SkeletonRemapping.GetSourceSkeletonBoneIndex(TargetSkeletonBoneIndex.GetInt()));
+						WeightIndex = SourceSkeletonBoneIndex.IsValid() ? SourceSkeletonBoneIndex.GetInt() : INDEX_NONE;
+					}
+					
+					(*SetWeights)[SetIndex].Value = WeightIndex != INDEX_NONE ? BoneMaskWeights[WeightIndex] : 0.0f;
+					(*SetWeights)[SetIndex].Value *= AdditiveWeight;
+				}
+			}
+			// Curve and attribute weights
+			else if (const TBoundValueMap<FFloatAnimationAttribute>* FloatMap = Cast<FFloatAnimationAttribute>(It.GetMap()))
+			{
+				SetWeights->FillWithValue(FFloatAnimationAttribute{ 0.0f });
+				
+				// Curves
+				const FAttributeTypedSetPtr& FloatSet = FloatMap->GetTypedSet();
+				CurveMaskWeights.ForEachElement([&FloatSet, SetWeights, ActiveAdditiveWeight = AdditiveWeight](const UE::Anim::FCurveElement& Curve)
+					{
+						const FAttributeSetIndex SetIndex = FloatSet->FindIndex(Curve.Name);
+						if (SetIndex.IsValid())
+						{
+							(*SetWeights)[SetIndex].Value = Curve.Value;
+							(*SetWeights)[SetIndex].Value *= ActiveAdditiveWeight;
+						}
+					});
+				
+				// Attributes
+				for (const FBlendProfileStandaloneCachedData::FMaskedAttributeWeight& AttributeWeight : AttributeMaskWeights)
+				{
+					const FAttributeSetIndex SetIndex = FloatSet->FindIndex(AttributeWeight.Attribute.GetName());
+					if (SetIndex.IsValid())
+					{
+						(*SetWeights)[SetIndex].Value = AttributeWeight.Weight;
+						(*SetWeights)[SetIndex].Value *= AdditiveWeight;
+					}
+				}
+			}
+			else
+			{
+				SetWeights->FillWithValue(FFloatAnimationAttribute{ 0.0f });
+
+				const FAttributeTypedSetPtr& TypedSet = It.GetMap()->GetTypedSet();
+				for (const FBlendProfileStandaloneCachedData::FMaskedAttributeWeight& AttributeWeight : AttributeMaskWeights)
+				{
+					const FAttributeSetIndex SetIndex = TypedSet->FindIndex(AttributeWeight.Attribute.GetName());
+					if (SetIndex.IsValid())
+					{
+						(*SetWeights)[SetIndex].Value = AttributeWeight.Weight;
+						(*SetWeights)[SetIndex].Value *= AdditiveWeight;
+					}
+				}
+			}
+		}
+	}
+}
diff --git a/Engine/Plugins/Experimental/UAF/UAFAnimNode/Source/UAFAnimNode/Public/UAF/AnimNodes/UAFApplyAdditive.h b/Engine/Plugins/Experimental/UAF/UAFAnimNode/Source/UAFAnimNode/Public/UAF/AnimNodes/UAFApplyAdditive.h
new file mode 100644
--- /dev/null
+++ b/Engine/Plugins/Experimental/UAF/UAFAnimNode/Source/UAFAnimNode/Public/UAF/AnimNodes/UAFApplyAdditive.h
@@ -0,0 +1,70 @@
+// Copyright Epic Games, Inc. All Rights Reserved.
+
+#pragma once
+
+#include "CoreMinimal.h"
+#include "BindableValue/UAFBindableTypes.h"
+#include "UAF/AnimNodeCore/UAFAnimNodeData.h"
+#include "UAF/AnimOps/UAFApplyAdditiveAnimOp.h"
+#include "UAF/BlendMask/UAFBlendMask.h"
+
+#include "UAFApplyAdditive.generated.h"
+
+#define UE_API UAFANIMNODE_API
+
+namespace UE::UAF
+{
+	// Node to apply an additive animation onto a base animation
+	// Optionally can use a blend mask for partial application of the additive
+	USTRUCT(DisplayName = "Apply Additive")
+	struct FUAFApplyAdditiveData : public FUAFAnimNodeData
+	{
+		GENERATED_BODY()
+		
+		// Node providing the base animation
+		UPROPERTY(EditAnywhere, Category = "NodeData")
+		TInstancedStruct<FUAFAnimNodeData> BaseNode;
+
+		// Node providing the additive animation
+		UPROPERTY(EditAnywhere, Category = "NodeData")
+		TInstancedStruct<FUAFAnimNodeData> AdditiveNode;
+		
+		// How much of the additive to apply to the base
+		UPROPERTY(EditAnywhere, Category = "NodeData")
+		FBindableFloat AdditiveWeight = 1.0f;
+		
+		// Optional blend mask to mask the additive with 
+		UPROPERTY(EditAnywhere, Category = "NodeData")
+		TObjectPtr<UUAFBlendMask> BlendMask = nullptr;
+		
+		UE_API virtual FUAFAnimNodePtr CreateInstance(FUAFAnimGraphUpdateContext& Context) const override;
+	};
+	
+	
+	class FUAFApplyAdditiveNode : public FUAFAnimNode
+	{
+	public:
+		UE_API explicit FUAFApplyAdditiveNode(FUAFAnimGraphUpdateContext& Context, const FUAFApplyAdditiveData& InData);
+		
+		// FUAFAnimNode impl
+		UE_API virtual void PreUpdate(FUAFAnimGraphUpdateContext& GraphContext) override;
+		UE_API virtual void* GetInterface(FUAFAnimNodeInterfaceId Id) override;
+		UE_API static void AddReferencedObjects(FUAFAnimNode* This, FReferenceCollector& Collector);
+		
+#if UAF_TRACE_ENABLED
+		UE_API virtual FString GetDebugName() const override;
+		UE_API virtual UStruct* GetDebugStruct() const override;
+#endif
+	
+	private:
+		static constexpr int32 BaseChildIndex = 0;
+		static constexpr int32 AdditiveChildIndex = 1;
+		
+		// Ptr to the nodes data - used to retrieve properties and evaluate bound values 
+		const FUAFApplyAdditiveData& Data;
+		// AnimOp for this node - performs the actual animation operations
+		FUAFApplyAdditiveAnimOp AdditiveAnimOp;
+	};
+}
+
+#undef UE_API 
diff --git a/Engine/Plugins/Experimental/UAF/UAFAnimNode/Source/UAFAnimNode/Public/UAF/AnimOps/UAFApplyAdditiveAnimOp.h b/Engine/Plugins/Experimental/UAF/UAFAnimNode/Source/UAFAnimNode/Public/UAF/AnimOps/UAFApplyAdditiveAnimOp.h
new file mode 100644
--- /dev/null
+++ b/Engine/Plugins/Experimental/UAF/UAFAnimNode/Source/UAFAnimNode/Public/UAF/AnimOps/UAFApplyAdditiveAnimOp.h
@@ -0,0 +1,46 @@
+// Copyright Epic Games, Inc. All Rights Reserved.
+
+#pragma once
+
+#include "CoreMinimal.h"
+#include "UAF/AnimOpCore/UAFAnimOp.h"
+#include "UAF/BlendMask/UAFBlendMask.h"
+
+#include "UAFApplyAdditiveAnimOp.generated.h"
+
+#define UE_API UAFANIMNODE_API
+
+namespace UE::UAF
+{
+	class FValueBundleStack;
+	class FValueBundle;
+
+	USTRUCT()
+	struct FUAFApplyAdditiveAnimOp : public FUAFAnimOp
+	{
+		GENERATED_BODY()
+		UAF_DECLARE_ANIMOP(FUAFApplyAdditiveAnimOp)
+		
+	public:
+		FUAFApplyAdditiveAnimOp();
+		
+		// FUAFAnimOp impl
+		UE_API virtual void EvaluateValues(FUAFAnimOpValueEvaluator& Evaluator) override;
+		
+		UE_API void SetAdditiveWeight(float InWeight);
+		UE_API void SetBlendMask(const TObjectPtr<UUAFBlendMask> InBlendMask);
+		
+	private:
+		void BuildPerValueWeights(const FUAFAnimOpValueEvaluator& Evaluator, const FValueBundle& BaseInput, FValueBundleStack& OutPerValueWeights);
+		
+	private:
+		// Optional blend mask for per-bone application of the additive
+		UPROPERTY(Transient)
+		TObjectPtr<UUAFBlendMask> OptionalBlendMask = nullptr;
+		
+		// The uniform weight to apply the additive with 
+		float AdditiveWeight = 1.0f;
+	};
+}
+
+#undef UE_API
