Update GetDebugName, GetDebugStruct on blend profile transition

#Synced-CL 52570859
#Synced-CL 52570859

diff --git a/Engine/Plugins/Experimental/UAF/UAFAnimNode/Source/UAFAnimNode/Private/UAF/AnimNodes/UAFBlendProfileTransition.cpp b/Engine/Plugins/Experimental/UAF/UAFAnimNode/Source/UAFAnimNode/Private/UAF/AnimNodes/UAFBlendProfileTransition.cpp
--- a/Engine/Plugins/Experimental/UAF/UAFAnimNode/Source/UAFAnimNode/Private/UAF/AnimNodes/UAFBlendProfileTransition.cpp
+++ b/Engine/Plugins/Experimental/UAF/UAFAnimNode/Source/UAFAnimNode/Private/UAF/AnimNodes/UAFBlendProfileTransition.cpp
@@ -59,13 +59,13 @@ namespace UE::UAF
 #if UAF_TRACE_ENABLED
 	FString FUAFBlendProfileTransition::GetDebugName() const
 	{
-		static FString ProfileTransitionName("Profile Transition");
-		return ProfileTransitionName;
+		static FString BlendProfileTransitionName("Blend Profile Transition");
+		return BlendProfileTransitionName;
 	}
 
 	UStruct* FUAFBlendProfileTransition::GetDebugStruct() const
 	{
-		return FUAFSimpleTransitionData::StaticStruct();
+		return FUAFBlendProfileTransitionData::StaticStruct();
 	}
 #endif
 }
