Remove unnecessary include

#Synced-CL 52539819

diff --git a/Engine/Plugins/Experimental/UAF/UAFAnimNode/Source/UAFAnimNode/Private/UAF/AnimNodes/UAFScalePlayRate.cpp b/Engine/Plugins/Experimental/UAF/UAFAnimNode/Source/UAFAnimNode/Private/UAF/AnimNodes/UAFScalePlayRate.cpp
--- a/Engine/Plugins/Experimental/UAF/UAFAnimNode/Source/UAFAnimNode/Private/UAF/AnimNodes/UAFScalePlayRate.cpp
+++ b/Engine/Plugins/Experimental/UAF/UAFAnimNode/Source/UAFAnimNode/Private/UAF/AnimNodes/UAFScalePlayRate.cpp
@@ -3,7 +3,6 @@
 #include "UAF/AnimNodes/UAFScalePlayRate.h"
 
 #include "UAF/AnimNodeCore/UAFAnimNodeUpdate.h"
-#include "UAF/AnimOps/UAFNullAnimOp.h"
 
 #include UE_INLINE_GENERATED_CPP_BY_NAME(UAFScalePlayRate)
 
