Convert RunAnimNode to use FValueBundle as its output

#Synced-CL 52587252

diff --git a/Engine/Plugins/Experimental/UAF/UAFAnimGraph/Source/UAFAnimGraph/Internal/Graph/VirtualValueBundle_ValueBundle.h b/Engine/Plugins/Experimental/UAF/UAFAnimGraph/Source/UAFAnimGraph/Internal/Graph/VirtualValueBundle_ValueBundle.h
new file mode 100644
--- /dev/null
+++ b/Engine/Plugins/Experimental/UAF/UAFAnimGraph/Source/UAFAnimGraph/Internal/Graph/VirtualValueBundle_ValueBundle.h
@@ -0,0 +1,44 @@
+// Copyright Epic Games, Inc. All Rights Reserved.
+
+#pragma once
+
+#include "Graph/AnimNext_LODPose.h"
+#include "Graph/WeakAnimGraphReference.h"
+#include "UAF/ValueRuntime/IVirtualValueBundle.h"
+#include "UAF/ValueRuntime/ValueBundle.h"
+
+#define UE_API UAFANIMGRAPH_API
+
+namespace UE::UAF
+{
+
+struct FVirtualValueBundle_ValueBundle : public IVirtualValueBundle
+{
+	explicit FVirtualValueBundle_ValueBundle(const FAnimNextGraphLODPose& InPose)
+		: Pose(InPose)
+	{}
+
+	explicit FVirtualValueBundle_ValueBundle(FAnimNextGraphLODPose&& InPose)
+		: Pose(MoveTemp(InPose))
+	{}
+
+	explicit FVirtualValueBundle_ValueBundle(const FValueBundleHeap& InValueBundle)
+		: ValueBundle(InValueBundle)
+	{}
+	
+		explicit FVirtualValueBundle_ValueBundle(FValueBundleHeap&& InValueBundle)
+		: ValueBundle(MoveTemp(InValueBundle))
+	{}
+	
+	// IVirtualValueBundle interface
+	UE_API virtual const FAnimNextGraphLODPose* GetLODPose() const override;
+	UE_API virtual const FValueBundleHeap* GetValueBundle() const override;
+
+private:
+	FAnimNextGraphLODPose Pose;
+	FValueBundleHeap ValueBundle;
+};
+
+}
+
+#undef UE_API
diff --git a/Engine/Plugins/Experimental/UAF/UAFAnimGraph/Source/UAFAnimGraph/Private/Graph/RigUnit_UAFRunAsset.cpp b/Engine/Plugins/Experimental/UAF/UAFAnimGraph/Source/UAFAnimGraph/Private/Graph/RigUnit_UAFRunAsset.cpp
--- a/Engine/Plugins/Experimental/UAF/UAFAnimGraph/Source/UAFAnimGraph/Private/Graph/RigUnit_UAFRunAsset.cpp
+++ b/Engine/Plugins/Experimental/UAF/UAFAnimGraph/Source/UAFAnimGraph/Private/Graph/RigUnit_UAFRunAsset.cpp
@@ -2,7 +2,7 @@
 
 #include "Graph/RigUnit_UAFRunAsset.h"
 
-#include "VirtualValueBundle_ValueBundle.h"
+#include "Graph/VirtualValueBundle_ValueBundle.h"
 #include "TraitInterfaces/IEvaluate.h"
 #include "TraitInterfaces/IUpdate.h"
 #include "Module/AnimNextModuleInstance.h"
diff --git a/Engine/Plugins/Experimental/UAF/UAFAnimGraph/Source/UAFAnimGraph/Private/Graph/VirtualValueBundle_ValueBundle.h b/Engine/Plugins/Experimental/UAF/UAFAnimGraph/Source/UAFAnimGraph/Private/Graph/VirtualValueBundle_ValueBundle.h
deleted file mode 100644
diff --git a/Engine/Plugins/Experimental/UAF/UAFAnimNode/Config/DefaultUAFAnimNode.ini b/Engine/Plugins/Experimental/UAF/UAFAnimNode/Config/DefaultUAFAnimNode.ini
new file mode 100644
--- /dev/null
+++ b/Engine/Plugins/Experimental/UAF/UAFAnimNode/Config/DefaultUAFAnimNode.ini
@@ -0,0 +1,4 @@
+[CoreRedirects]
+; Converting output to abstract skeleton runtime
++StructRedirects=(OldName="/Script/UAFAnimNode.RigUnit_RunAnimNode", NewName="/Script/UAFAnimNode.RigUnit_RunAnimNode_v1")
++StructRedirects=(OldName="/Script/UAFAnimNode.RunAnimNodeWorkData", NewName="/Script/UAFAnimNode.RunAnimNodeWorkData_v1")
diff --git a/Engine/Plugins/Experimental/UAF/UAFAnimNode/Source/UAFAnimNode/Internal/AnimNode/RigUnit_RunAnimNode.h b/Engine/Plugins/Experimental/UAF/UAFAnimNode/Source/UAFAnimNode/Internal/AnimNode/RigUnit_RunAnimNode.h
deleted file mode 100644
diff --git a/Engine/Plugins/Experimental/UAF/UAFAnimNode/Source/UAFAnimNode/Internal/AnimNode/RigUnit_RunAnimNode_v1.h b/Engine/Plugins/Experimental/UAF/UAFAnimNode/Source/UAFAnimNode/Internal/AnimNode/RigUnit_RunAnimNode_v1.h
new file mode 100644
--- /dev/null
+++ b/Engine/Plugins/Experimental/UAF/UAFAnimNode/Source/UAFAnimNode/Internal/AnimNode/RigUnit_RunAnimNode_v1.h
@@ -0,0 +1,77 @@
+// Copyright Epic Games, Inc. All Rights Reserved.
+
+#pragma once
+
+#include "AnimNextExecuteContext.h"
+#include "Graph/AnimNext_LODPose.h"
+#include "Graph/RigUnit_AnimNextBase.h"
+#include "UAF/AnimNodeCore/UAFAnimNode.h"
+#include "UAF/AnimNodeCore/UAFAnimNodeDataEx.h"
+#include "UAF/AnimNodeCore/UAFAnimNodeReferenceCollector.h"
+
+#include "RigUnit_RunAnimNode_v1.generated.h"
+
+#define UE_API UAFANIMNODE_API
+
+// TODO: Remove me once old content has been resaved
+USTRUCT()
+struct FUAFAnimNodeReference
+{
+	GENERATED_BODY()
+	UPROPERTY(EditAnywhere, Category = "Run", meta = (Input, ExcludeBaseStruct))
+	TInstancedStruct<UE::UAF::FUAFAnimNodeData> NodeData;
+};
+
+USTRUCT()
+struct FRunAnimNodeWorkData_v1
+{
+	GENERATED_BODY()
+
+	UE_API ~FRunAnimNodeWorkData_v1();
+
+	UE::UAF::FUAFAnimNodePtr AnimNodeInstance;
+
+	TSharedPtr<UE::UAF::FUAFAnimNodeReferenceCollector> GCReferences;
+
+	uint64 DebugInstanceId = 0;
+};
+
+/** Runs an AnimNode */
+USTRUCT(meta=(Hidden, DisplayName="Run Anim Node (Deprecated)", Category="Anim Node", NodeColor="0, 1, 1", Keywords="Animation, Node, Graph, AnimNode"))
+struct FRigUnit_RunAnimNode_v1 : public FRigUnit_AnimNextBase
+{
+	GENERATED_BODY()
+
+	RIGVM_METHOD()
+	UE_API void Execute();
+
+	// The anim node to run
+	UPROPERTY(EditAnywhere, Category = "Run", meta = (Input))
+	UE::UAF::FUAFAnimNodeDataEx AnimNode;
+
+	// LOD to run at. If this is -1 then the reference pose's source LOD will be used
+	UPROPERTY(EditAnywhere, Category = "Run", meta = (Input))
+	int32 LOD = -1;
+
+	// Reference pose
+	UPROPERTY(EditAnywhere, Category = "Run", meta = (Input))
+	FAnimNextGraphReferencePose ReferencePose;
+
+	// The name of the attribute set to execute with, if 'None' will include everything
+	UPROPERTY(EditAnywhere, Category = "Run", meta = (Input))
+	FName AttributeSet;
+
+	// Pose result
+	UPROPERTY(EditAnywhere, Category = "Run", meta = (Output))
+	FAnimNextGraphLODPose Result;
+
+	// Internal work data
+	UPROPERTY(transient)
+	FRunAnimNodeWorkData_v1 WorkData;
+
+	// The execution result
+	UPROPERTY(EditAnywhere, DisplayName = "Execute", Category = "BeginExecution", meta = (Input, Output))
+	FAnimNextExecuteContext ExecuteContext;
+};
+
+#undef UE_API
diff --git a/Engine/Plugins/Experimental/UAF/UAFAnimNode/Source/UAFAnimNode/Internal/AnimNode/RigUnit_RunAnimNode_v2.h b/Engine/Plugins/Experimental/UAF/UAFAnimNode/Source/UAFAnimNode/Internal/AnimNode/RigUnit_RunAnimNode_v2.h
new file mode 100644
--- /dev/null
+++ b/Engine/Plugins/Experimental/UAF/UAFAnimNode/Source/UAFAnimNode/Internal/AnimNode/RigUnit_RunAnimNode_v2.h
@@ -0,0 +1,60 @@
+// Copyright Epic Games, Inc. All Rights Reserved.
+
+#pragma once
+
+#include "AnimNextExecuteContext.h"
+#include "Graph/RigUnit_AnimNextBase.h"
+#include "UAF/AnimNodeCore/UAFAnimNode.h"
+#include "UAF/AnimNodeCore/UAFAnimNodeDataEx.h"
+#include "UAF/AnimNodeCore/UAFAnimNodeReferenceCollector.h"
+#include "UAF/ValueRuntime/ValueBundle.h"
+
+#include "RigUnit_RunAnimNode_v2.generated.h"
+
+#define UE_API UAFANIMNODE_API
+
+USTRUCT()
+struct FRunAnimNodeWorkData_v2
+{
+	GENERATED_BODY()
+
+	UE_API ~FRunAnimNodeWorkData_v2();
+
+	UE::UAF::FUAFAnimNodePtr AnimNodeInstance;
+
+	TSharedPtr<UE::UAF::FUAFAnimNodeReferenceCollector> GCReferences;
+
+	uint64 DebugInstanceId = 0;
+};
+
+/** Runs an AnimNode */
+USTRUCT(meta=(DisplayName="Run Anim Node", Category="Anim Node", NodeColor="0, 1, 1", Keywords="Animation, Node, Graph, AnimNode"))
+struct FRigUnit_RunAnimNode_v2 : public FRigUnit_AnimNextBase
+{
+	GENERATED_BODY()
+
+	RIGVM_METHOD()
+	UE_API void Execute();
+
+	// The anim node to run
+	UPROPERTY(EditAnywhere, Category = "Run", meta = (Input))
+	UE::UAF::FUAFAnimNodeDataEx AnimNode;
+
+	// The name of the attribute set to execute with, if 'None' will include everything
+	UPROPERTY(EditAnywhere, Category = "Run", meta = (Input))
+	FName AttributeSet;
+
+	// Result value
+	UPROPERTY(EditAnywhere, Category = "Run", meta = (Output))
+	FUAFValueBundle Result;
+
+	// Internal work data
+	UPROPERTY(transient)
+	FRunAnimNodeWorkData_v2 WorkData;
+
+	// The execution result
+	UPROPERTY(EditAnywhere, DisplayName = "Execute", Category = "BeginExecution", meta = (Input, Output))
+	FAnimNextExecuteContext ExecuteContext;
+};
+
+#undef UE_API
diff --git a/Engine/Plugins/Experimental/UAF/UAFAnimNode/Source/UAFAnimNode/Private/RigUnit_RunAnimNode.cpp b/Engine/Plugins/Experimental/UAF/UAFAnimNode/Source/UAFAnimNode/Private/RigUnit_RunAnimNode.cpp
deleted file mode 100644
diff --git a/Engine/Plugins/Experimental/UAF/UAFAnimNode/Source/UAFAnimNode/Private/RigUnit_RunAnimNode_v1.cpp b/Engine/Plugins/Experimental/UAF/UAFAnimNode/Source/UAFAnimNode/Private/RigUnit_RunAnimNode_v1.cpp
new file mode 100644
--- /dev/null
+++ b/Engine/Plugins/Experimental/UAF/UAFAnimNode/Source/UAFAnimNode/Private/RigUnit_RunAnimNode_v1.cpp
@@ -0,0 +1,149 @@
+// Copyright Epic Games, Inc. All Rights Reserved.
+
+#include "AnimNode/RigUnit_RunAnimNode_v1.h"
+
+#include "Engine/SkeletalMesh.h"
+#include "Graph/UAFSystemOutputComponent.h"
+#include "Module/AnimNextModuleInstance.h"
+#include "Module/UAFNotifyDispatcherComponent.h"
+#include "UAF/UAFSkeletonUserData.h"
+#include "UAFLogging.h"
+#include "UAF/AnimNodeCore/UAFAnimNodeTrace.h"
+#include "UAF/AnimNodeCore/UAFAnimNodeUpdate.h"
+#include "ObjectTrace.h"
+
+#include UE_INLINE_GENERATED_CPP_BY_NAME(RigUnit_RunAnimNode_v1)
+
+FRunAnimNodeWorkData_v1::~FRunAnimNodeWorkData_v1()
+{
+	// Clear our instance before we destroy the reference collector
+	if (AnimNodeInstance)
+	{
+		UE::UAF::FUAFAnimGraphUpdateContext Context(nullptr, nullptr, DebugInstanceId, *GCReferences, 0.0f); 
+		AnimNodeInstance = nullptr;
+	}
+
+	GCReferences = nullptr;
+}
+
+FRigUnit_RunAnimNode_v1_Execute()
+{
+	SCOPED_NAMED_EVENT(UAF_Run_AnimNode, FColor::Orange);
+
+	using namespace UE::UAF;
+
+	const FAnimNextModuleContextData& ModuleContextData = ExecuteContext.GetContextData<FAnimNextModuleContextData>();
+	FAnimNextModuleInstance& ModuleInstance = ModuleContextData.GetModuleInstance();
+
+	if (!ReferencePose.ReferencePose.IsValid())
+	{
+		UAF_RIGUNIT_LOG(Warning, TEXT("Could not run Anim Node - Ref Pose is invalid]")); 
+		return;
+	}
+	
+	if (!AnimNode.IsValid())
+	{
+		UAF_RIGUNIT_LOG(Warning, TEXT("Could not run Anim Node - No Anim Node Selected]")); 
+		return;
+	}
+
+#if UAF_TRACE_ENABLED
+	if (WorkData.DebugInstanceId == 0)
+	{
+		// Trace an Instance to be the outer containing track in rewind debugger for all the node instance tracks
+		// Using this RigUnit's struct type so we can register a custom track details panel to that type in future
+		WorkData.DebugInstanceId = FObjectTrace::AllocateInstanceId();
+		TRACE_INSTANCE(ModuleContextData.GetObject(), WorkData.DebugInstanceId, ModuleInstance.GetUniqueId(), FRigUnit_RunAnimNode_v1::StaticStruct(), "UAF Anim Graph");
+	}
+#endif
+
+	const UE::UAF::FReferencePose& RefPose = ReferencePose.ReferencePose.GetRef<UE::UAF::FReferencePose>();
+
+	int32 DesiredLOD = LOD;
+	if (DesiredLOD == -1)
+	{
+		DesiredLOD = RefPose.GetSourceLODLevel();
+	}
+
+	// TODO: Currently forcing additive flag to false here
+	if (Result.LODPose.ShouldPrepareForLOD(RefPose, DesiredLOD, false))
+	{
+		Result.LODPose.PrepareForLOD(RefPose, DesiredLOD, true, false);
+	}
+
+	ensure(Result.LODPose.LODLevel == DesiredLOD);
+
+	// Internally we use memstack allocation, so we need a mark here
+	FMemStack& MemStack = FMemStack::Get();
+	FMemMark MemMark(MemStack);
+
+	// We allocate a dummy buffer to trigger the allocation of a large chunk if this is the first mark
+	// This reduces churn internally by avoiding a chunk to be repeatedly allocated and freed as we push/pop marks
+	MemStack.Alloc(size_t(FPageAllocator::SmallPageSize) + 1, 16);
+
+	if (WorkData.GCReferences == nullptr)
+	{
+		WorkData.GCReferences = MakeShared<FUAFAnimNodeReferenceCollector>();
+	}
+
+	FUAFAnimGraphUpdateContext Context(ModuleContextData.GetObject(), &ModuleInstance, WorkData.DebugInstanceId, *WorkData.GCReferences, ExecuteContext.GetDeltaTime());
+
+	if (WorkData.AnimNodeInstance == nullptr)
+	{
+		// allocate AnimNode instance
+		WorkData.AnimNodeInstance = AnimNode.Get()->CreateInstance(Context);
+	}
+
+	TRACE_UAF_MODULE(ModuleInstance);
+
+	if (WorkData.AnimNodeInstance)
+	{
+		const TArray<FUAFAnimOp*> AnimOpList = UpdateGraph(Context, *WorkData.AnimNodeInstance);
+
+		UAF_TRACE_ANIMOPS(AnimOpList, ModuleInstance, ModuleInstance.GetObject());
+
+		// Use AnimOps to compute our pose
+		{
+			FUAFSystemOutputComponent& OutputPoseComponent = ModuleInstance.GetOrAddComponent<FUAFSystemOutputComponent>();
+
+			// Note that the skeletal mesh and skeleton can be null
+			USkeletalMesh* SkeletalMesh = OutputPoseComponent.GetSkeletalMesh();
+			USkeleton* Skeleton = SkeletalMesh != nullptr ? SkeletalMesh->GetSkeleton() : nullptr;
+
+			if (Skeleton != nullptr)
+			{
+				TNonNullPtr<UUAFSkeletonUserData> UserData = UUAFSkeletonUserData::FromSkeleton(Skeleton);
+				UAbstractSkeletonSetBinding* Binding = UserData->GetSetBinding();
+
+				FUAFEvaluateValuesArgs EvaluateArgs(*Binding, AttributeSet, SkeletalMesh, DesiredLOD, AnimOpList, RefPose);
+
+				EvaluateValues(EvaluateArgs, Result);
+			}
+			else
+			{
+				// If we have no skeletal mesh or skeleton, we output nothing
+			}
+		}
+
+		// Now that everything has updated and we've computed our values, we know for certain every blend weight
+		
+		// Use AnimOps to evaluate our notifies
+		{
+			TArray<FAnimNotifyEventReference> Notifies = EvaluateNotifies(AnimOpList);
+			if (!Notifies.IsEmpty())
+			{
+				// Ensure we have a handler component on the module
+				(void)ModuleInstance.GetOrAddComponent<FUAFNotifyDispatcherComponent>();
+
+				auto NotifyDispatchEvent = MakeTraitEvent<FNotifyDispatchEvent>();
+				NotifyDispatchEvent->Notifies = MoveTemp(Notifies);
+				NotifyDispatchEvent->Weight = 1.0f;
+
+				ModuleInstance.QueueOutputTraitEvent(NotifyDispatchEvent);
+			}
+		}
+
+		// Evaluate sync contributors now while the AnimOps are ready, they will be used in the next update
+		EvaluateSyncContributors(AnimOpList);
+	}
+}
diff --git a/Engine/Plugins/Experimental/UAF/UAFAnimNode/Source/UAFAnimNode/Private/RigUnit_RunAnimNode_v2.cpp b/Engine/Plugins/Experimental/UAF/UAFAnimNode/Source/UAFAnimNode/Private/RigUnit_RunAnimNode_v2.cpp
new file mode 100644
--- /dev/null
+++ b/Engine/Plugins/Experimental/UAF/UAFAnimNode/Source/UAFAnimNode/Private/RigUnit_RunAnimNode_v2.cpp
@@ -0,0 +1,147 @@
+// Copyright Epic Games, Inc. All Rights Reserved.
+
+#include "AnimNode/RigUnit_RunAnimNode_v2.h"
+
+#include "Engine/SkeletalMesh.h"
+#include "Graph/UAFSystemOutputComponent.h"
+#include "Module/AnimNextModuleInstance.h"
+#include "Module/UAFNotifyDispatcherComponent.h"
+#include "UAF/UAFSkeletonUserData.h"
+#include "UAFLogging.h"
+#include "UAF/AnimNodeCore/UAFAnimNodeTrace.h"
+#include "UAF/AnimNodeCore/UAFAnimNodeUpdate.h"
+#include "ObjectTrace.h"
+#include "Graph/VirtualValueBundle_ValueBundle.h"
+
+#include UE_INLINE_GENERATED_CPP_BY_NAME(RigUnit_RunAnimNode_v2)
+
+FRunAnimNodeWorkData_v2::~FRunAnimNodeWorkData_v2()
+{
+	// Clear our instance before we destroy the reference collector
+	if (AnimNodeInstance)
+	{
+		UE::UAF::FUAFAnimGraphUpdateContext Context(nullptr, nullptr, DebugInstanceId, *GCReferences, 0.0f); 
+		AnimNodeInstance = nullptr;
+	}
+
+	GCReferences = nullptr;
+}
+
+FRigUnit_RunAnimNode_v2_Execute()
+{
+	SCOPED_NAMED_EVENT(UAF_Run_AnimNode, FColor::Orange);
+
+	using namespace UE::UAF;
+
+	const FAnimNextModuleContextData& ModuleContextData = ExecuteContext.GetContextData<FAnimNextModuleContextData>();
+	FAnimNextModuleInstance& ModuleInstance = ModuleContextData.GetModuleInstance();
+	
+	if (!AnimNode.IsValid())
+	{
+		UAF_RIGUNIT_LOG(Warning, TEXT("Could not run Anim Node - No Anim Node Selected]")); 
+		return;
+	}
+
+#if UAF_TRACE_ENABLED
+	if (WorkData.DebugInstanceId == 0)
+	{
+		// Trace an Instance to be the outer containing track in rewind debugger for all the node instance tracks
+		// Using this RigUnit's struct type so we can register a custom track details panel to that type in future
+		WorkData.DebugInstanceId = FObjectTrace::AllocateInstanceId();
+		TRACE_INSTANCE(ModuleContextData.GetObject(), WorkData.DebugInstanceId, ModuleInstance.GetUniqueId(), FRigUnit_RunAnimNode_v2::StaticStruct(), "UAF Anim Graph");
+	}
+#endif
+
+	FUAFSystemOutputComponent& OutputComponent = ModuleInstance.GetOrAddComponent<FUAFSystemOutputComponent>();
+	FDataHandle RefPoseDataHandle = OutputComponent.GetRefPose().ReferencePose;
+	if (!RefPoseDataHandle.IsValid())
+	{
+		UAF_RIGUNIT_LOG(Warning, TEXT("Could not run anim node - Ref Pose is invalid."));
+		return;
+	}
+
+	const UE::UAF::FReferencePose& RefPose = RefPoseDataHandle.GetRef<UE::UAF::FReferencePose>();
+
+	const int32 DesiredLOD = RefPose.GetSourceLODLevel();
+
+	FAnimNextGraphLODPose GraphPose;
+	// TODO: Currently forcing additive flag to false here
+	GraphPose.LODPose.PrepareForLOD(RefPose, DesiredLOD, true, false);
+
+	ensure(GraphPose.LODPose.LODLevel == DesiredLOD);
+
+	// Internally we use memstack allocation, so we need a mark here
+	FMemStack& MemStack = FMemStack::Get();
+	FMemMark MemMark(MemStack);
+
+	// We allocate a dummy buffer to trigger the allocation of a large chunk if this is the first mark
+	// This reduces churn internally by avoiding a chunk to be repeatedly allocated and freed as we push/pop marks
+	MemStack.Alloc(size_t(FPageAllocator::SmallPageSize) + 1, 16);
+
+	if (WorkData.GCReferences == nullptr)
+	{
+		WorkData.GCReferences = MakeShared<FUAFAnimNodeReferenceCollector>();
+	}
+
+	FUAFAnimGraphUpdateContext Context(ModuleContextData.GetObject(), &ModuleInstance, WorkData.DebugInstanceId, *WorkData.GCReferences, ExecuteContext.GetDeltaTime());
+
+	if (WorkData.AnimNodeInstance == nullptr)
+	{
+		// allocate AnimNode instance
+		WorkData.AnimNodeInstance = AnimNode.Get()->CreateInstance(Context);
+	}
+
+	TRACE_UAF_MODULE(ModuleInstance);
+
+	if (WorkData.AnimNodeInstance)
+	{
+		const TArray<FUAFAnimOp*> AnimOpList = UpdateGraph(Context, *WorkData.AnimNodeInstance);
+
+		UAF_TRACE_ANIMOPS(AnimOpList, ModuleInstance, ModuleInstance.GetObject());
+
+		// Use AnimOps to compute our pose
+		{
+			// Note that the skeletal mesh and skeleton can be null
+			USkeletalMesh* SkeletalMesh = OutputComponent.GetSkeletalMesh();
+			USkeleton* Skeleton = SkeletalMesh != nullptr ? SkeletalMesh->GetSkeleton() : nullptr;
+
+			if (Skeleton != nullptr)
+			{
+				TNonNullPtr<UUAFSkeletonUserData> UserData = UUAFSkeletonUserData::FromSkeleton(Skeleton);
+				UAbstractSkeletonSetBinding* Binding = UserData->GetSetBinding();
+
+				FUAFEvaluateValuesArgs EvaluateArgs(*Binding, AttributeSet, SkeletalMesh, DesiredLOD, AnimOpList, RefPose);
+
+				EvaluateValues(EvaluateArgs, GraphPose);
+			}
+			else
+			{
+				// If we have no skeletal mesh or skeleton, we output nothing
+			}
+		}
+
+		// Now that everything has updated and we've computed our values, we know for certain every blend weight
+		
+		// Use AnimOps to evaluate our notifies
+		{
+			TArray<FAnimNotifyEventReference> Notifies = EvaluateNotifies(AnimOpList);
+			if (!Notifies.IsEmpty())
+			{
+				// Ensure we have a handler component on the module
+				(void)ModuleInstance.GetOrAddComponent<FUAFNotifyDispatcherComponent>();
+
+				auto NotifyDispatchEvent = MakeTraitEvent<FNotifyDispatchEvent>();
+				NotifyDispatchEvent->Notifies = MoveTemp(Notifies);
+				NotifyDispatchEvent->Weight = 1.0f;
+
+				ModuleInstance.QueueOutputTraitEvent(NotifyDispatchEvent);
+			}
+		}
+
+		// Evaluate sync contributors now while the AnimOps are ready, they will be used in the next update
+		EvaluateSyncContributors(AnimOpList);
+	}
+
+	// Move our output pose into the result
+	Result.SetAs<FVirtualValueBundle_ValueBundle>(MoveTemp(GraphPose));
+}
