Handle empty value bundle in modifier AnimOps

#Synced-CL 52678682

diff --git a/Engine/Plugins/Experimental/UAF/UAFAnimNode/Source/UAFAnimNode/Private/UAF/AnimOps/UAFOffsetRootBoneAnimOp.cpp b/Engine/Plugins/Experimental/UAF/UAFAnimNode/Source/UAFAnimNode/Private/UAF/AnimOps/UAFOffsetRootBoneAnimOp.cpp
--- a/Engine/Plugins/Experimental/UAF/UAFAnimNode/Source/UAFAnimNode/Private/UAF/AnimOps/UAFOffsetRootBoneAnimOp.cpp
+++ b/Engine/Plugins/Experimental/UAF/UAFAnimNode/Source/UAFAnimNode/Private/UAF/AnimOps/UAFOffsetRootBoneAnimOp.cpp
@@ -81,7 +81,12 @@ void FUAFOffsetRootBoneAnimOp::EvaluateValues(FUAFAnimOpValueEvaluator& Evaluato
 	// Access the value bundle from the evaluation context
 	const FUAFAnimOpValueEvaluationContext& AnimOpEvaluationContext = Evaluator.GetActiveEvaluationContext();
 	FPoseValueBundleCoWRef* ValueBundle = Evaluator.GetEvaluationStack().PeekMutable();
-	check (ValueBundle);
+	check(ValueBundle);
+
+	if (ValueBundle->Get().IsEmpty())
+	{
+		return;
+	}
 	
 	FAttributeNamedSetPtr NamedSet = ValueBundle->Get().GetNamedSet();
 	if (!NamedSet.IsValid())
diff --git a/Engine/Plugins/Experimental/UAF/UAFAnimNode/Source/UAFAnimNode/Private/UAF/AnimOps/UAFSteeringAnimOp.cpp b/Engine/Plugins/Experimental/UAF/UAFAnimNode/Source/UAFAnimNode/Private/UAF/AnimOps/UAFSteeringAnimOp.cpp
--- a/Engine/Plugins/Experimental/UAF/UAFAnimNode/Source/UAFAnimNode/Private/UAF/AnimOps/UAFSteeringAnimOp.cpp
+++ b/Engine/Plugins/Experimental/UAF/UAFAnimNode/Source/UAFAnimNode/Private/UAF/AnimOps/UAFSteeringAnimOp.cpp
@@ -50,7 +50,12 @@ void FUAFSteeringAnimOp::EvaluateValues(FUAFAnimOpValueEvaluator& Evaluator)
 	const FUAFAnimOpValueEvaluationContext& AnimOpEvaluationContext = Evaluator.GetActiveEvaluationContext();
 
 	FPoseValueBundleCoWRef* ValueBundle = Evaluator.GetEvaluationStack().PeekMutable();
-	check (ValueBundle != nullptr)
+	check(ValueBundle != nullptr);
+
+	if (ValueBundle->Get().IsEmpty())
+	{
+		return;
+	}
 	
 	FAttributeNamedSetPtr NamedSet = ValueBundle->Get().GetNamedSet();
 	if (NamedSet.IsValid() == false)
diff --git a/Engine/Plugins/Experimental/UAF/UAFWarping/Source/UAFWarping/Private/AnimNode/UAFStrafeWarpingNode.cpp b/Engine/Plugins/Experimental/UAF/UAFWarping/Source/UAFWarping/Private/AnimNode/UAFStrafeWarpingNode.cpp
--- a/Engine/Plugins/Experimental/UAF/UAFWarping/Source/UAFWarping/Private/AnimNode/UAFStrafeWarpingNode.cpp
+++ b/Engine/Plugins/Experimental/UAF/UAFWarping/Source/UAFWarping/Private/AnimNode/UAFStrafeWarpingNode.cpp
@@ -576,6 +576,16 @@ void FUAFStrafeWarpingPostAnimOp::EvaluateValues(FUAFAnimOpValueEvaluator& Evalu
 		return;
 	}
 
+	{
+		const FPoseValueBundleCoWRef* ValueBundle = Evaluator.GetEvaluationStack().Peek();
+		check(ValueBundle != nullptr);
+
+		if (ValueBundle->Get().IsEmpty())
+		{
+			return;
+		}
+	}
+
 	float TargetOrientationAngleRad = 0.0f;
 	const FVector RotationAxisVector = FUAFStrafeWarpingUtils::GetAxisVector(SharedData->RotationAxis);
 
diff --git a/Engine/Plugins/Experimental/UAF/UAFWarping/Source/UAFWarping/Private/AnimNode/UAFWarpToTargetNode.cpp b/Engine/Plugins/Experimental/UAF/UAFWarping/Source/UAFWarping/Private/AnimNode/UAFWarpToTargetNode.cpp
--- a/Engine/Plugins/Experimental/UAF/UAFWarping/Source/UAFWarping/Private/AnimNode/UAFWarpToTargetNode.cpp
+++ b/Engine/Plugins/Experimental/UAF/UAFWarping/Source/UAFWarping/Private/AnimNode/UAFWarpToTargetNode.cpp
@@ -26,6 +26,11 @@ void FUAFWarpToTargetPostAnimOp::EvaluateValues(FUAFAnimOpValueEvaluator& Evalua
 	FPoseValueBundleCoWRef* InputRef = Evaluator.GetEvaluationStack().PeekMutable();
 	FPoseValueBundle& ValueBundle = InputRef->GetMutable();
 
+	if (ValueBundle.IsEmpty())
+	{
+		return;
+	}
+
 	const FAttributeNamedSetPtr NamedSet = ValueBundle.GetNamedSet();
 	if (!NamedSet)
 	{
