Add support for tracing struct anim node values in UAF

#Synced-CL 52661316

diff --git a/Engine/Plugins/Experimental/UAF/UAFAnimNode/Source/UAFAnimNode/Private/UAF/AnimNodeCore/UAFAnimNodeTrace.cpp b/Engine/Plugins/Experimental/UAF/UAFAnimNode/Source/UAFAnimNode/Private/UAF/AnimNodeCore/UAFAnimNodeTrace.cpp
--- a/Engine/Plugins/Experimental/UAF/UAFAnimNode/Source/UAFAnimNode/Private/UAF/AnimNodeCore/UAFAnimNodeTrace.cpp
+++ b/Engine/Plugins/Experimental/UAF/UAFAnimNode/Source/UAFAnimNode/Private/UAF/AnimNodeCore/UAFAnimNodeTrace.cpp
@@ -38,6 +38,7 @@ UE_TRACE_EVENT_BEGIN(UAFAnimNode, AnimNodeValue)
 	UE_TRACE_EVENT_FIELD(uint64, NodeId)
 	UE_TRACE_EVENT_FIELD(uint32, NameId)
 	UE_TRACE_EVENT_FIELD(uint8, Type)
+	UE_TRACE_EVENT_FIELD(uint64, StructType)
 	UE_TRACE_EVENT_FIELD(uint8[], Value)
 UE_TRACE_EVENT_END()
 
@@ -105,7 +106,7 @@ namespace UE::UAF
 		}
 	}
 
-	void TraceAnimNodeValue(FUAFAnimGraphUpdateContext& Context, FUAFAnimNode* Node, EPropertyBagPropertyType Type, FName PropertyName, TConstArrayView<uint8> Value)
+	void TraceAnimNodeValue(FUAFAnimGraphUpdateContext& Context, FUAFAnimNode* Node, EPropertyBagPropertyType Type, const UScriptStruct* StructType, FName PropertyName, TConstArrayView<uint8> Value)
 	{
 		if (UE_TRACE_CHANNELEXPR_IS_ENABLED(UAFChannel))
 		{
@@ -123,6 +124,12 @@ namespace UE::UAF
 				}
 
 				uint32 NameId = FAnimTrace::OutputName(PropertyName);
+				uint64 StructId = 0;
+				if (StructType)
+				{
+					TRACE_TYPE(StructType);
+					StructId = FObjectTrace::GetObjectId(StructType);
+				}
 	
 				UE_TRACE_LOG(UAFAnimNode, AnimNodeValue, UAFChannel)
 					<< AnimNodeValue.Cycle(FPlatformTime::Cycles64())
@@ -130,6 +137,7 @@ namespace UE::UAF
 					<< AnimNodeValue.NodeId(Node->GetDebugInstanceId())
 					<< AnimNodeValue.RootGraphId( Context.GetOuterDebugInstanceId())
 					<< AnimNodeValue.Type(static_cast<uint8>(Type))
+					<< AnimNodeValue.StructType(StructId)
 					<< AnimNodeValue.NameId(NameId)
 					<< AnimNodeValue.Value(Value.GetData(), Value.Num());
 			}
@@ -138,18 +146,31 @@ namespace UE::UAF
 
 	void TraceAnimNodeValue(FUAFAnimGraphUpdateContext& Context, FUAFAnimNode* Node, FName PropertyName, bool Value)
 	{
-		TraceAnimNodeValue(Context, Node, EPropertyBagPropertyType::Bool, PropertyName, TConstArrayView<uint8>(reinterpret_cast<uint8*>(&Value), sizeof(bool)));
+		TraceAnimNodeValue(Context, Node, EPropertyBagPropertyType::Bool, nullptr, PropertyName, TConstArrayView<uint8>(reinterpret_cast<uint8*>(&Value), sizeof(bool)));
 	}
 
 	void TraceAnimNodeValue(FUAFAnimGraphUpdateContext& Context, FUAFAnimNode* Node, FName PropertyName, double Value)
 	{
-		TraceAnimNodeValue(Context, Node, EPropertyBagPropertyType::Double, PropertyName, TConstArrayView<uint8>(reinterpret_cast<uint8*>(&Value), sizeof(double)));
+		TraceAnimNodeValue(Context, Node, EPropertyBagPropertyType::Double, nullptr, PropertyName, TConstArrayView<uint8>(reinterpret_cast<uint8*>(&Value), sizeof(double)));
+	}
+
+	void TraceAnimNodeValue(FUAFAnimGraphUpdateContext& Context, FUAFAnimNode* Node, FName PropertyName, const FStructView Value)
+	{
+		if (const UScriptStruct* StructType = Value.GetScriptStruct())
+		{
+			TArray<uint8> ArchiveData;
+			FMemoryWriter WriterArchive(ArchiveData);
+			FObjectAsTraceIdProxyArchive Archive(WriterArchive);
+			StructType->SerializeBin(Archive, Value.GetMemory());
+		
+			TraceAnimNodeValue(Context, Node, EPropertyBagPropertyType::Struct, StructType, PropertyName, ArchiveData);
+		}
 	}
 
 	void TraceAnimNodeValue(FUAFAnimGraphUpdateContext& Context, FUAFAnimNode* Node, FName PropertyName, const UObject* Value)
 	{
 		uint64 ObjectId = FObjectTrace::GetObjectId(Value);
-		TraceAnimNodeValue(Context, Node, EPropertyBagPropertyType::Object, PropertyName, TConstArrayView<uint8>(reinterpret_cast<uint8*>(&ObjectId), sizeof(uint64)));
+		TraceAnimNodeValue(Context, Node, EPropertyBagPropertyType::Object, nullptr, PropertyName, TConstArrayView<uint8>(reinterpret_cast<uint8*>(&ObjectId), sizeof(uint64)));
 	}
 }
 #endif
diff --git a/Engine/Plugins/Experimental/UAF/UAFAnimNode/Source/UAFAnimNode/Private/UAF/AnimNodes/UAFSteeringNode.cpp b/Engine/Plugins/Experimental/UAF/UAFAnimNode/Source/UAFAnimNode/Private/UAF/AnimNodes/UAFSteeringNode.cpp
--- a/Engine/Plugins/Experimental/UAF/UAFAnimNode/Source/UAFAnimNode/Private/UAF/AnimNodes/UAFSteeringNode.cpp
+++ b/Engine/Plugins/Experimental/UAF/UAFAnimNode/Source/UAFAnimNode/Private/UAF/AnimNodes/UAFSteeringNode.cpp
@@ -57,6 +57,9 @@ void FUAFSteeringNode::PreUpdate(FUAFAnimGraphUpdateContext& GraphContext)
 	SteeringAnimOp.MinScaleRatio = Data->MinScaleRatio.GetValue(GraphContext.GetVariablesOwner());
 	SteeringAnimOp.MaxScaleRatio = Data->MaxScaleRatio.GetValue(GraphContext.GetVariablesOwner());
 	
+	UAF_TRACE_ANIMNODE_VALUE(GraphContext, this, "Alpha", SteeringAnimOp.Alpha);
+	UAF_TRACE_ANIMNODE_VALUE(GraphContext, this, "TargetOrientation", FStructView::Make(SteeringAnimOp.TargetOrientation));
+	
 	FUAFAnimNode::PreUpdate(GraphContext);
 }
 
diff --git a/Engine/Plugins/Experimental/UAF/UAFAnimNode/Source/UAFAnimNode/Public/UAF/AnimNodeCore/UAFAnimNodeTrace.h b/Engine/Plugins/Experimental/UAF/UAFAnimNode/Source/UAFAnimNode/Public/UAF/AnimNodeCore/UAFAnimNodeTrace.h
--- a/Engine/Plugins/Experimental/UAF/UAFAnimNode/Source/UAFAnimNode/Public/UAF/AnimNodeCore/UAFAnimNodeTrace.h
+++ b/Engine/Plugins/Experimental/UAF/UAFAnimNode/Source/UAFAnimNode/Public/UAF/AnimNodeCore/UAFAnimNodeTrace.h
@@ -4,6 +4,7 @@
 
 #include "CoreMinimal.h"
 #include "RewindDebugger/UAFTrace.h"
+#include "StructUtils/StructView.h"
 
 #define UE_API UAFANIMNODE_API
 
@@ -23,6 +24,7 @@ namespace UE::UAF
 	UE_API void TraceAnimNodeValue(FUAFAnimGraphUpdateContext& Context, FUAFAnimNode* Node, FName PropertyName, bool Value);
 	UE_API void TraceAnimNodeValue(FUAFAnimGraphUpdateContext& Context, FUAFAnimNode* Node, FName PropertyName, double Value);
 	UE_API void TraceAnimNodeValue(FUAFAnimGraphUpdateContext& Context, FUAFAnimNode* Node, FName PropertyName, const UObject* Value);
+	UE_API void TraceAnimNodeValue(FUAFAnimGraphUpdateContext& Context, FUAFAnimNode* Node, FName PropertyName, FStructView Value);
 }
 
 #define UAF_TRACE_ANIMOPS(AnimOps, AssetInstance, Object) UE::UAF::TraceAnimOps(AnimOps, AssetInstance, Object);
diff --git a/Engine/Plugins/Experimental/UAF/UAFAnimNode/Source/UAFAnimNodeEditor/Private/RewindDebugger/UAFAnimNodeAnalyzer.cpp b/Engine/Plugins/Experimental/UAF/UAFAnimNode/Source/UAFAnimNodeEditor/Private/RewindDebugger/UAFAnimNodeAnalyzer.cpp
--- a/Engine/Plugins/Experimental/UAF/UAFAnimNode/Source/UAFAnimNodeEditor/Private/RewindDebugger/UAFAnimNodeAnalyzer.cpp
+++ b/Engine/Plugins/Experimental/UAF/UAFAnimNode/Source/UAFAnimNodeEditor/Private/RewindDebugger/UAFAnimNodeAnalyzer.cpp
@@ -67,8 +67,9 @@ bool FUAFAnimNodeAnalyzer::OnEvent(uint16 RouteId, EStyle Style, const FOnEventC
 			uint64 NodeId = EventData.GetValue<uint64>("NodeId");
 			uint32 NameId = EventData.GetValue<uint32>("NameId");
 			uint8 Type = EventData.GetValue<uint8>("Type");
+			uint64 StructType = EventData.GetValue<uint64>("StructType");
 			TArrayView<const uint8> Value = EventData.GetArrayView<uint8>("Value");
-			Provider.AppendAnimNodeValue(Context.EventTime.AsSeconds(Cycle), RecordingTime, RootGraphId, NodeId, Type, NameId, Value);
+			Provider.AppendAnimNodeValue(Context.EventTime.AsSeconds(Cycle), RecordingTime, RootGraphId, NodeId, Type, StructType, NameId, Value);
 			break;			
 		}
 
diff --git a/Engine/Plugins/Experimental/UAF/UAFAnimNode/Source/UAFAnimNodeEditor/Private/RewindDebugger/UAFAnimNodeProvider.cpp b/Engine/Plugins/Experimental/UAF/UAFAnimNode/Source/UAFAnimNodeEditor/Private/RewindDebugger/UAFAnimNodeProvider.cpp
--- a/Engine/Plugins/Experimental/UAF/UAFAnimNode/Source/UAFAnimNodeEditor/Private/RewindDebugger/UAFAnimNodeProvider.cpp
+++ b/Engine/Plugins/Experimental/UAF/UAFAnimNode/Source/UAFAnimNodeEditor/Private/RewindDebugger/UAFAnimNodeProvider.cpp
@@ -1,6 +1,8 @@
 // Copyright Epic Games, Inc. All Rights Reserved.
 
 #include "UAFAnimNodeProvider.h"
+
+#include "IGameplayProvider.h"
 #include "ObjectTrace.h"
 
 namespace TraceServices
@@ -55,7 +57,7 @@ void FUAFAnimNodeProvider::AppendAnimNodeUpdate(double ProfileTime, double Recor
 }
 
 
-void FUAFAnimNodeProvider::AppendAnimNodeValue(double ProfileTime, double RecordingTime, uint64 RootGraphId, uint64 NodeId, uint8 Type, uint32 NameId, TConstArrayView<uint8> Value)
+void FUAFAnimNodeProvider::AppendAnimNodeValue(double ProfileTime, double RecordingTime, uint64 RootGraphId, uint64 NodeId, uint8 Type, uint64 StructType, uint32 NameId, TConstArrayView<uint8> Value)
 {
 	EditAccessCheck();
     
@@ -75,6 +77,16 @@ void FUAFAnimNodeProvider::AppendAnimNodeValue(double ProfileTime, double Record
 	ValueData.NodeId = NodeId;
 	ValueData.NameId = NameId;
 	ValueData.Type = Type;
+	if (StructType != 0)
+	{
+		if (const IGameplayProvider* GameplayProvider = Session.ReadProvider<IGameplayProvider>("GameplayProvider"))
+		{
+			GameplayProvider->BeginRead();
+			const FObjectInfo& StructInfo = GameplayProvider->GetObjectInfo(StructType);
+			GameplayProvider->EndRead();
+			ValueData.StructType = FindObject<UScriptStruct>(nullptr, StructInfo.PathName);
+		}
+	}
 	ValueData.Value = Value;
 
 	(*Data)->NodeValueTimeline.AppendEvent(ProfileTime, ValueData);
diff --git a/Engine/Plugins/Experimental/UAF/UAFAnimNode/Source/UAFAnimNodeEditor/Private/RewindDebugger/UAFAnimNodeProvider.h b/Engine/Plugins/Experimental/UAF/UAFAnimNode/Source/UAFAnimNodeEditor/Private/RewindDebugger/UAFAnimNodeProvider.h
--- a/Engine/Plugins/Experimental/UAF/UAFAnimNode/Source/UAFAnimNodeEditor/Private/RewindDebugger/UAFAnimNodeProvider.h
+++ b/Engine/Plugins/Experimental/UAF/UAFAnimNode/Source/UAFAnimNodeEditor/Private/RewindDebugger/UAFAnimNodeProvider.h
@@ -67,6 +67,7 @@ struct FAnimNodeValueData
 	uint32 NameId;
 	uint8 Type;
 	TArray<uint8, TInlineAllocator<8>> Value;
+	const UScriptStruct* StructType = nullptr;
 };
 
 
@@ -127,7 +128,7 @@ public:
 	
 	void AppendAnimOpList(double ProfileTime, double RecordingTime, uint64 ObjectInstanceId, uint64 GraphInstanceId, const TArrayView<const uint8>& ListData);
 	void AppendAnimNodeUpdate(double ProfileTime, double RecordingTime, uint64 RootGraphId, uint64 NodeId, uint64 ParentNodeId, float TotalWeight);
-	void AppendAnimNodeValue(double ProfileTime, double RecordingTime, uint64 RootGraphId, uint64 NodeId, uint8 TypeId, uint32 NameId, TConstArrayView<uint8> Value);
+	void AppendAnimNodeValue(double ProfileTime, double RecordingTime, uint64 RootGraphId, uint64 NodeId, uint8 TypeId, uint64 StructType, uint32 NameId, TConstArrayView<uint8> Value);
 	
 	void EnumerateEvaluationGraphs(uint64 OuterObjectId, TFunctionRef<void(uint64 GraphInstanceId)> Callback) const;
 	
diff --git a/Engine/Plugins/Experimental/UAF/UAFAnimNode/Source/UAFAnimNodeEditor/Private/RewindDebugger/UAFAnimNodeTrack.cpp b/Engine/Plugins/Experimental/UAF/UAFAnimNode/Source/UAFAnimNodeEditor/Private/RewindDebugger/UAFAnimNodeTrack.cpp
--- a/Engine/Plugins/Experimental/UAF/UAFAnimNode/Source/UAFAnimNodeEditor/Private/RewindDebugger/UAFAnimNodeTrack.cpp
+++ b/Engine/Plugins/Experimental/UAF/UAFAnimNode/Source/UAFAnimNodeEditor/Private/RewindDebugger/UAFAnimNodeTrack.cpp
@@ -16,6 +16,7 @@
 #include "UAF/AnimNodeCore/UAFAnimNodeData.h"
 #include "UAF/AnimNodeCore/UAFTransitionNodeData.h"
 #include "IAnimationProvider.h"
+#include "ObjectAsTraceIdProxyArchiveReader.h"
 #include "Misc/PackageName.h"
 
 #include UE_INLINE_GENERATED_CPP_BY_NAME(UAFAnimNodeTrack)
@@ -206,8 +207,13 @@ bool FUAFAnimNodeTrack::UpdateInternal()
 						if (NodeValueData.NodeId == NodeId)
 						{
 							FName PropertyName = AnimationProvider->GetName(NodeValueData.NameId);
+							
 							FPropertyBagPropertyDesc PropertyDesc(PropertyName, static_cast<EPropertyBagPropertyType>(NodeValueData.Type), UObject::StaticClass());
-
+							if (NodeValueData.StructType)
+							{
+								PropertyDesc.ValueTypeObject = NodeValueData.StructType;
+							}
+							
 							DetailsObject->Properties.AddProperties({ PropertyDesc });
 
 							switch (static_cast<EPropertyBagPropertyType>(NodeValueData.Type))
@@ -215,13 +221,13 @@ bool FUAFAnimNodeTrack::UpdateInternal()
 							case EPropertyBagPropertyType::Bool:
 								if (ensure(NodeValueData.Value.Num() == sizeof(bool)))
 								{
-									DetailsObject->Properties.SetValueBool(PropertyName, *reinterpret_cast<const bool*>(&NodeValueData.Value[0]));
+									DetailsObject->Properties.SetValueBool(PropertyName, *reinterpret_cast<const bool*>(NodeValueData.Value.GetData()));
 								}
 								break;
 							case EPropertyBagPropertyType::Double:
 								if (ensure(NodeValueData.Value.Num() == sizeof(double)))
 								{
-									DetailsObject->Properties.SetValueDouble(PropertyName, *reinterpret_cast<const double*>(&NodeValueData.Value[0]));
+									DetailsObject->Properties.SetValueDouble(PropertyName, *reinterpret_cast<const double*>(NodeValueData.Value.GetData()));
 								}
 								break;
 							case EPropertyBagPropertyType::Object:
@@ -229,7 +235,7 @@ bool FUAFAnimNodeTrack::UpdateInternal()
 								if (ensure(NodeValueData.Value.Num() == sizeof(uint64)))
 								{
 									UObject* Asset = nullptr;
-									const uint64 ObjectId = *reinterpret_cast<const uint64*>(&NodeValueData.Value[0]);
+									const uint64 ObjectId = *reinterpret_cast<const uint64*>(NodeValueData.Value.GetData());
 									const FObjectInfo& AssetInfo = GameplayProvider->GetObjectInfo(ObjectId);
 									const FString PackagePathString = FPackageName::ObjectPathToPackageName(FString(AssetInfo.PathName));
 
@@ -251,6 +257,24 @@ bool FUAFAnimNodeTrack::UpdateInternal()
 								}
 								break;
 							}
+							case EPropertyBagPropertyType::Struct:
+							{
+								if (NodeValueData.StructType)
+								{
+									TArray<uint8> DeserializedStruct;
+									DeserializedStruct.SetNum(NodeValueData.StructType->GetStructureSize());
+									NodeValueData.StructType->InitializeStruct(DeserializedStruct.GetData());
+									TArray<uint8> ValueArray(NodeValueData.Value.GetData(), NodeValueData.Value.Num());
+									FMemoryReader StructReader(ValueArray);
+									FObjectAsTraceIdProxyArchiveReader Reader(StructReader, GameplayProvider);
+									NodeValueData.StructType->SerializeBin(Reader, DeserializedStruct.GetData());
+								
+									DetailsObject->Properties.SetValueStruct(PropertyName,FConstStructView(NodeValueData.StructType, DeserializedStruct.GetData()));
+									
+									NodeValueData.StructType->DestroyStruct(DeserializedStruct.GetData());
+								}
+								break;
+							}
 							default:
 								break;
 							};
