Add InputValue anim node/op

Leaf node that reads a FUAFValueBundle graph variable during PreUpdate and pushes it onto the evaluation stack. Prepares for the direct FValueBundle path and already fully supports the legacy LODPose path with automatic cross-skeleton bone remapping.

#Synced-CL 52682769

diff --git a/Engine/Plugins/Experimental/UAF/UAF/Source/UAF/Private/Graph/AnimNext_LODPose.cpp b/Engine/Plugins/Experimental/UAF/UAF/Source/UAF/Private/Graph/AnimNext_LODPose.cpp
--- a/Engine/Plugins/Experimental/UAF/UAF/Source/UAF/Private/Graph/AnimNext_LODPose.cpp
+++ b/Engine/Plugins/Experimental/UAF/UAF/Source/UAF/Private/Graph/AnimNext_LODPose.cpp
@@ -11,7 +11,7 @@ void FAnimNextGraphLODPose::CopyPose(
 	const UE::UAF::FReferencePose* SourceRef = InSourcePose.RefPose;
 	const UE::UAF::FReferencePose* TargetRef = OutTargetPose.RefPose;
 
-	if (SourceRef != nullptr && TargetRef != nullptr && SourceRef != TargetRef)
+	if (SourceRef != nullptr && TargetRef != nullptr && UE::UAF::FRemapPoseDataPool::NeedsRemapping(*SourceRef, *TargetRef))
 	{
 		const FRemapPoseData& RemapData = UE::UAF::FRemapPoseDataPool::Get().GetRemapData(*SourceRef, *TargetRef);
 		RemapData.RemapPose(InSourcePose, OutTargetPose);
diff --git a/Engine/Plugins/Experimental/UAF/UAF/Source/UAF/Public/RemapPoseDataPool.h b/Engine/Plugins/Experimental/UAF/UAF/Source/UAF/Public/RemapPoseDataPool.h
--- a/Engine/Plugins/Experimental/UAF/UAF/Source/UAF/Public/RemapPoseDataPool.h
+++ b/Engine/Plugins/Experimental/UAF/UAF/Source/UAF/Public/RemapPoseDataPool.h
@@ -22,6 +22,15 @@ public:
 	/** Access the singleton pool instance. */
 	static UE_API FRemapPoseDataPool& Get();
 
+	/**
+	 * Returns true if the source and target reference poses refer to different skeletons, meaning a copy between them requires remapping.
+	 * This is a cheap reference pose pointer comparison, no locking or pool access.
+	 */
+	[[nodiscard]] static bool NeedsRemapping(const FReferencePose& InSourceRefPose, const FReferencePose& InTargetRefPose)
+	{
+		return &InSourceRefPose != &InTargetRefPose;
+	}
+
 	/**
 	 * Get or create remap data for a source->target reference pose pair.
 	 * On first access for a given pair, FRemapPoseData::Reinit() is called to build the bone mapping. Subsequent calls with the same pair return the cached data.
diff --git a/Engine/Plugins/Experimental/UAF/UAFAnimGraph/Source/UAFAnimGraph/Private/EvaluationVM/KeyframeState.cpp b/Engine/Plugins/Experimental/UAF/UAFAnimGraph/Source/UAFAnimGraph/Private/EvaluationVM/KeyframeState.cpp
--- a/Engine/Plugins/Experimental/UAF/UAFAnimGraph/Source/UAFAnimGraph/Private/EvaluationVM/KeyframeState.cpp
+++ b/Engine/Plugins/Experimental/UAF/UAFAnimGraph/Source/UAFAnimGraph/Private/EvaluationVM/KeyframeState.cpp
@@ -21,7 +21,7 @@ namespace UE::UAF
 		const FReferencePose* SourceRef = Source.LODPose.RefPose;
 		const FReferencePose* TargetRef = Pose.RefPose;
 
-		if (SourceRef != nullptr && TargetRef != nullptr && SourceRef != TargetRef)
+		if (SourceRef != nullptr && TargetRef != nullptr && FRemapPoseDataPool::NeedsRemapping(*SourceRef, *TargetRef))
 		{
 			// Cross-skeleton: Fill ref pose transforms for unmapped bones, then remap mapped ones.
 			Pose.SetRefPose(Source.LODPose.IsAdditive());
diff --git a/Engine/Plugins/Experimental/UAF/UAFAnimNode/Source/UAFAnimNode/Private/UAF/AnimNodes/InputValueAnimNode.cpp b/Engine/Plugins/Experimental/UAF/UAFAnimNode/Source/UAFAnimNode/Private/UAF/AnimNodes/InputValueAnimNode.cpp
new file mode 100644
--- /dev/null
+++ b/Engine/Plugins/Experimental/UAF/UAFAnimNode/Source/UAFAnimNode/Private/UAF/AnimNodes/InputValueAnimNode.cpp
@@ -0,0 +1,214 @@
+// Copyright Epic Games, Inc. All Rights Reserved.
+
+#include "UAF/AnimNodes/InputValueAnimNode.h"
+
+#include "Graph/AnimNext_LODPose.h"
+#include "Graph/UAFSystemOutputComponent.h"
+#include "Module/AnimNextModuleInstance.h"
+#include "RemapPoseDataPool.h"
+#include "UAF/AnimNodeCore/UAFAnimNodeUpdate.h"
+#include "UAF/AnimOpCore/UAFAnimOpValueEvaluator.h"
+#include "UAF/Attributes/EngineAttributes.h"
+#include "UAF/ValueRuntime/PoseValueBundle.h"
+#include "UAF/ValueRuntime/IVirtualValueBundle.h"
+#include "UAFAssetInstance.h"
+
+#include UE_INLINE_GENERATED_CPP_BY_NAME(InputValueAnimNode)
+
+namespace UE::UAF
+{
+
+void FUAFInputValueAnimOp::ConvertLODPoseToValueBundle(
+	const FAnimNextGraphLODPose& InputPose,
+	const TReferencePose<FDefaultAllocator, FDefaultSetAllocator>* TargetRefPose,
+	FPoseValueBundleStack& OutValueBundle)
+{
+	const FLODPoseHeap& SourcePose = InputPose.LODPose;
+	const FReferencePose* SourceRefPose = SourcePose.RefPose;
+
+	// Determine if we need to remap between different skeletons
+	const bool bNeedsRemap = SourceRefPose != nullptr && TargetRefPose != nullptr
+		&& FRemapPoseDataPool::NeedsRemapping(*SourceRefPose, *TargetRefPose);
+
+	// The pose we'll actually convert which is either the source directly or a remapped copy.
+	const FLODPose* PoseToConvert = &SourcePose;
+	FLODPoseHeap RemappedPose;
+
+	if (bNeedsRemap)
+	{
+		const FRemapPoseData& RemapData = FRemapPoseDataPool::Get().GetRemapData(*SourceRefPose, *TargetRefPose);
+		RemapData.RemapPose(SourcePose, RemappedPose);
+		PoseToConvert = &RemappedPose;
+	}
+
+	OutValueBundle.InitWithValueSpace(FValueSpace(EValueSpaceType::Local));
+	FPoseValueBundle& Pose = FPoseValueBundle::From(OutValueBundle);
+
+	// Convert bone transforms
+	if (TBoundValueMap<FBoneTransformAnimationAttribute>* BoneTransforms = Pose.FindBoneTransforms())
+	{
+		const FAttributeTypedSetPtr& TypedSet = BoneTransforms->GetTypedSet();
+		const FReferencePose& RefPose = PoseToConvert->GetRefPose();
+		const int32 NumLODBones = PoseToConvert->GetNumBones();
+
+		for (int32 LODBoneIndex = 0; LODBoneIndex < NumLODBones; ++LODBoneIndex)
+		{
+			const int32 SkeletonBoneIndex = RefPose.GetSkeletonBoneIndexFromLODBoneIndex(LODBoneIndex);
+			const FAttributeBindingIndex BoneBindingIndex = FAttributeBindingIndex(FSkeletonPoseBoneIndex(SkeletonBoneIndex));
+			if (BoneBindingIndex.IsValid())
+			{
+				const FAttributeSetIndex BoneSetIndex = TypedSet->GetIndex(BoneBindingIndex);
+				if (BoneSetIndex.IsValid())
+				{
+					(*BoneTransforms)[BoneSetIndex].Value = PoseToConvert->LocalTransformsView[LODBoneIndex];
+				}
+			}
+		}
+	}
+
+	// Convert float curves (name-matched, no skeleton dependency)
+	if (TBoundValueMap<FFloatAnimationAttribute>* FloatCurves = Pose.FindFloatCurves())
+	{
+		const FAttributeTypedSetPtr& TypedSet = FloatCurves->GetTypedSet();
+		const FBlendedHeapCurve& Curves = InputPose.Curves;
+
+		Curves.ForEachElement([&FloatCurves, &TypedSet](const auto& CurveElement)
+		{
+			const FAttributeSetIndex CurveSetIndex = TypedSet->FindIndex(CurveElement.Name);
+			if (CurveSetIndex.IsValid())
+			{
+				(*FloatCurves)[CurveSetIndex].Value = CurveElement.Value;
+			}
+		});
+	}
+}
+
+// ---------------------------------------------------------------------------
+// FUAFInputValueAnimOp
+// ---------------------------------------------------------------------------
+
+FUAFInputValueAnimOp::FUAFInputValueAnimOp()
+	: FUAFAnimOp(0)
+{
+	InitializeAs<FUAFInputValueAnimOp>();
+}
+
+void FUAFInputValueAnimOp::EvaluateValues(FUAFAnimOpValueEvaluator& Evaluator)
+{
+	const FAttributeNamedSetPtr& NamedSet = Evaluator.GetActiveEvaluationContext().GetNamedSet();
+
+	// Direct FValueBundle path, push an immutable reference, CoW handles copying if needed
+	// TODO: Verify this path once the system output is an actual value bundle.
+	if (CachedBundle && !CachedBundle->IsEmpty())
+	{
+		Evaluator.GetEvaluationStack().Push(UE::UAF::FPoseValueBundleCoWRef::MakeImmutable(static_cast<const FPoseValueBundle*>(CachedBundle)));
+		return;
+	}
+
+	// LODPose path, remap if needed, then convert into a FValueBundle
+	if (CachedLODPose && CachedLODPose->LODPose.IsValid() && NamedSet)
+	{
+		FPoseValueBundleStack OutValueBundle(NamedSet);
+		ConvertLODPoseToValueBundle(*CachedLODPose, TargetRefPose, OutValueBundle);
+		Evaluator.GetEvaluationStack().Push(UE::UAF::FPoseValueBundleCoWRef::MakeFrom(MoveTemp(OutValueBundle)));
+		return;
+	}
+
+	// Fallback, push an empty bundle (reference pose)
+	FPoseValueBundleStack EmptyBundle(NamedSet);
+	Evaluator.GetEvaluationStack().Push(UE::UAF::FPoseValueBundleCoWRef::MakeFrom(MoveTemp(EmptyBundle)));
+}
+
+// ---------------------------------------------------------------------------
+// FUAFInputValueAnimNodeData
+// ---------------------------------------------------------------------------
+
+FUAFAnimNodePtr FUAFInputValueAnimNodeData::CreateInstance(FUAFAnimGraphUpdateContext& Context) const
+{
+	return MakeAnimNode<FUAFInputValueAnimNode>(Context, *this);
+}
+
+// ---------------------------------------------------------------------------
+// FUAFInputValueAnimNode
+// ---------------------------------------------------------------------------
+
+FUAFInputValueAnimNode::FUAFInputValueAnimNode(FUAFAnimGraphUpdateContext& Context, const FUAFInputValueAnimNodeData& InData)
+	: FUAFAnimNode(Context)
+{
+	InitializeAs<FUAFInputValueAnimNode>(Context);
+
+	VariableReference = InData.Input;
+
+	InputValueOp.SetDebugOwner(this);
+	SetPostAnimOp(&InputValueOp);
+}
+
+void FUAFInputValueAnimNode::PreUpdate(FUAFAnimGraphUpdateContext& GraphContext)
+{
+	// Clear cached pointers each frame, only valid for one evaluation cycle
+	InputValueOp.CachedBundle = nullptr;
+	InputValueOp.CachedLODPose = nullptr;
+	InputValueOp.TargetRefPose = nullptr;
+
+	FUAFAssetInstance* VariablesOwner = GraphContext.GetVariablesOwner();
+	if (!VariablesOwner || VariableReference.IsNone())
+	{
+		return;
+	}
+
+	// Cache the target reference pose from the module's output component.
+	// Used during EvaluateValues to remap bones if the source and target skeletons differ.
+	if (FAnimNextModuleInstance* ModuleInstance = VariablesOwner->GetRootInstance())
+	{
+		if (const FUAFSystemOutputComponent* OutputComponent = ModuleInstance->TryGetComponent<FUAFSystemOutputComponent>())
+		{
+			const FAnimNextGraphReferencePose& GraphRefPose = OutputComponent->GetRefPose();
+			if (GraphRefPose.ReferencePose.IsValid())
+			{
+				InputValueOp.TargetRefPose = &GraphRefPose.ReferencePose.GetRef<TReferencePose<FDefaultAllocator, FDefaultSetAllocator>>();
+			}
+		}
+	}
+
+	// Read the FUAFValueBundle variable and cache whichever representation is available.
+	// This must happen during PreUpdate because variable access is only available
+	// during the update phase, not during the evaluation phase.
+	VariablesOwner->AccessVariable<FUAFValueBundle>(VariableReference, [this](FUAFValueBundle& InValueBundle)
+	{
+		if (const IVirtualValueBundle* Impl = InValueBundle.GetImpl())
+		{
+			// Try the direct FValueBundle path first (new attribute runtime)
+			if (const FValueBundleHeap* ValueBundle = Impl->GetValueBundle())
+			{
+				if (!ValueBundle->IsEmpty())
+				{
+					InputValueOp.CachedBundle = ValueBundle;
+					return;
+				}
+			}
+
+			// Fall back to the LODPose path (old pose system)
+			if (const FAnimNextGraphLODPose* LODPose = Impl->GetLODPose())
+			{
+				if (LODPose->LODPose.IsValid())
+				{
+					InputValueOp.CachedLODPose = LODPose;
+				}
+			}
+		}
+	});
+}
+
+#if UAF_TRACE_ENABLED
+FString FUAFInputValueAnimNode::GetDebugName() const
+{
+	return FString::Printf(TEXT("InputValue: %s"), *VariableReference.GetName().ToString());
+}
+
+UStruct* FUAFInputValueAnimNode::GetDebugStruct() const
+{
+	return FUAFInputValueAnimNodeData::StaticStruct();
+}
+#endif
+
+} // namespace UE::UAF
diff --git a/Engine/Plugins/Experimental/UAF/UAFAnimNode/Source/UAFAnimNode/Public/UAF/AnimNodes/InputValueAnimNode.h b/Engine/Plugins/Experimental/UAF/UAFAnimNode/Source/UAFAnimNode/Public/UAF/AnimNodes/InputValueAnimNode.h
new file mode 100644
--- /dev/null
+++ b/Engine/Plugins/Experimental/UAF/UAFAnimNode/Source/UAFAnimNode/Public/UAF/AnimNodes/InputValueAnimNode.h
@@ -0,0 +1,97 @@
+// Copyright Epic Games, Inc. All Rights Reserved.
+
+#pragma once
+
+#include "CoreMinimal.h"
+#include "ReferencePose.h"
+#include "UAF/AnimNodeCore/UAFAnimNode.h"
+#include "UAF/AnimNodeCore/UAFAnimNodeData.h"
+#include "UAF/AnimOpCore/UAFAnimOp.h"
+#include "UAF/ValueRuntime/PoseValueBundle.h"
+#include "UAF/ValueRuntime/ValueBundle.h"
+#include "Variables/AnimNextVariableReference.h"
+
+#include "InputValueAnimNode.generated.h"
+
+#define UE_API UAFANIMNODE_API
+
+struct FAnimNextGraphLODPose;
+
+namespace UE::UAF
+{
+
+/**
+ * FUAFInputValueAnimOp
+ * A leaf AnimOp that pushes a cached value bundle or LOD pose onto the evaluation stack.
+ * The cached pointers are set by FUAFInputValueAnimNode during PreUpdate because variable
+ * access is only available during the update phase, not during the evaluation phase.
+ * If no data is available, an empty bundle (reference pose) is pushed instead.
+ */
+USTRUCT()
+struct FUAFInputValueAnimOp : public FUAFAnimOp
+{
+	GENERATED_BODY()
+	UAF_DECLARE_ANIMOP(FUAFInputValueAnimOp)
+
+	FUAFInputValueAnimOp();
+
+	virtual void EvaluateValues(FUAFAnimOpValueEvaluator& Evaluator) override;
+
+	/**
+	 * Cached data set by FUAFInputValueAnimNode::PreUpdate each frame.
+	 * Non-owning pointers, only valid for the current evaluation cycle.
+	 */
+	const FValueBundle* CachedBundle = nullptr;
+	const FAnimNextGraphLODPose* CachedLODPose = nullptr;
+	const TReferencePose<FDefaultAllocator, FDefaultSetAllocator>* TargetRefPose = nullptr;
+	
+	/** Converts an FAnimNextGraphLODPose into a FValueBundleStack, with optional skeleton remapping. */
+	static void ConvertLODPoseToValueBundle(const FAnimNextGraphLODPose& InputPose,
+		const TReferencePose<FDefaultAllocator, FDefaultSetAllocator>* TargetRefPose,
+		FPoseValueBundleStack& OutValueBundle);
+};
+
+/**
+ * FUAFInputValueAnimNodeData
+ * Serializable data for the input value anim node.
+ */
+USTRUCT(DisplayName = "Input Value")
+struct UE_API FUAFInputValueAnimNodeData : public FUAFAnimNodeData
+{
+	GENERATED_BODY()
+
+	/** The graph variable to read the input value bundle from. */
+	UPROPERTY(EditAnywhere, Category = "Input Value", meta = (AllowedType = FUAFValueBundle))
+	FAnimNextVariableReference Input;
+
+	virtual FUAFAnimNodePtr CreateInstance(FUAFAnimGraphUpdateContext& Context) const override;
+};
+
+/**
+ * FUAFInputValueAnimNode
+ * An AnimNode leaf that reads a FUAFValueBundle from a graph variable during PreUpdate and pushes it onto the evaluation stack via FUAFInputValueAnimOp.
+ */
+struct UE_API FUAFInputValueAnimNode : public FUAFAnimNode
+{
+	FUAFInputValueAnimNode(FUAFAnimGraphUpdateContext& Context, const FUAFInputValueAnimNodeData& InData);
+
+	// FUAFAnimNode interface
+	virtual void* GetInterface(FUAFAnimNodeInterfaceId Id) override { return nullptr; }
+	virtual void PreUpdate(FUAFAnimGraphUpdateContext& GraphContext) override;
+
+#if UAF_TRACE_ENABLED
+	virtual FString GetDebugName() const override;
+	virtual UStruct* GetDebugStruct() const override;
+#endif
+
+private:
+	/** The AnimOp that pushes the cached bundle onto the eval stack. */
+	FUAFInputValueAnimOp InputValueOp;
+
+	/** Variable reference to read from. */
+	FAnimNextVariableReference VariableReference;
+};
+
+} // namespace UE::UAF
+
+#undef UE_API
