Add construction variants to AnimOp ArrayView

#Synced-CL 52678682

diff --git a/Engine/Plugins/Experimental/UAF/UAFAnimNode/Source/UAFAnimNode/Public/UAF/AnimOpCore/UAFAnimOpArrayView.h b/Engine/Plugins/Experimental/UAF/UAFAnimNode/Source/UAFAnimNode/Public/UAF/AnimOpCore/UAFAnimOpArrayView.h
--- a/Engine/Plugins/Experimental/UAF/UAFAnimNode/Source/UAFAnimNode/Public/UAF/AnimOpCore/UAFAnimOpArrayView.h
+++ b/Engine/Plugins/Experimental/UAF/UAFAnimNode/Source/UAFAnimNode/Public/UAF/AnimOpCore/UAFAnimOpArrayView.h
@@ -21,9 +21,11 @@ namespace UE::UAF
 
 		// Creates an array view from a C-style array
 		template<size_t Count> static FUAFAnimOpArrayView Make(FUAFAnimOp*(&AnimOps)[Count]);
+		template<size_t Count> static FUAFAnimOpArrayView Make(FUAFAnimOp* (&AnimOps)[Count], int32 Num);
 
 		// Creates an array view from a TArray with any allocator
 		template<class AllocatorType> static FUAFAnimOpArrayView Make(const TArray<FUAFAnimOp*, AllocatorType>& AnimOps);
+		template<class AllocatorType> static FUAFAnimOpArrayView Make(const TArray<FUAFAnimOp*, AllocatorType>& AnimOps, int32 Num);
 
 		// FUAFAnimOp impl
 		UE_API virtual void EvaluateValues(FUAFAnimOpValueEvaluator& Evaluator) override;
@@ -55,6 +57,16 @@ namespace UE::UAF
 		return Result;
 	}
 
+	template<size_t Count>
+	inline FUAFAnimOpArrayView FUAFAnimOpArrayView::Make(FUAFAnimOp* (&AnimOps)[Count], int32 Num)
+	{
+		checkf(Num <= Count, TEXT("Cannot create a view that extends past the end of the input array."));
+
+		FUAFAnimOpArrayView Result;
+		Result.AnimOps = MakeArrayView(AnimOps, Num);
+		return Result;
+	}
+
 	template<class AllocatorType>
 	inline FUAFAnimOpArrayView FUAFAnimOpArrayView::Make(const TArray<FUAFAnimOp*, AllocatorType>& AnimOps)
 	{
@@ -62,6 +74,16 @@ namespace UE::UAF
 		Result.AnimOps = MakeArrayView(AnimOps);
 		return Result;
 	}
+
+	template<class AllocatorType>
+	inline FUAFAnimOpArrayView FUAFAnimOpArrayView::Make(const TArray<FUAFAnimOp*, AllocatorType>& AnimOps, int32 Num)
+	{
+		checkf(Num <= AnimOps.Num(), TEXT("Cannot create a view that extends past the end of the input array."));
+
+		FUAFAnimOpArrayView Result;
+		Result.AnimOps = MakeArrayView(AnimOps.GetData(), Num);
+		return Result;
+	}
 }
 
 #undef UE_API
