Add accessor to set the play rate

#Synced-CL 52678682

diff --git a/Engine/Plugins/Experimental/UAF/UAFAnimNode/Source/UAFAnimNode/Private/UAF/AnimNodes/UAFScalePlayRate.cpp b/Engine/Plugins/Experimental/UAF/UAFAnimNode/Source/UAFAnimNode/Private/UAF/AnimNodes/UAFScalePlayRate.cpp
--- a/Engine/Plugins/Experimental/UAF/UAFAnimNode/Source/UAFAnimNode/Private/UAF/AnimNodes/UAFScalePlayRate.cpp
+++ b/Engine/Plugins/Experimental/UAF/UAFAnimNode/Source/UAFAnimNode/Private/UAF/AnimNodes/UAFScalePlayRate.cpp
@@ -16,6 +16,12 @@ namespace UE::UAF
 		InitializeModifier(Context, Data);
 	}
 
+	void FUAFScalePlayRate::SetPlayRate(float PlayRate)
+	{
+		PlayRateMultiplier.ClearBinding();
+		PlayRateMultiplier.SetConstantValue(PlayRate);
+	}
+
 	void FUAFScalePlayRate::PreUpdate(FUAFAnimGraphUpdateContext& Context)
 	{
 		float Multiplier = PlayRateMultiplier.GetValue(Context.GetVariablesOwner());
diff --git a/Engine/Plugins/Experimental/UAF/UAFAnimNode/Source/UAFAnimNode/Public/UAF/AnimNodes/UAFScalePlayRate.h b/Engine/Plugins/Experimental/UAF/UAFAnimNode/Source/UAFAnimNode/Public/UAF/AnimNodes/UAFScalePlayRate.h
--- a/Engine/Plugins/Experimental/UAF/UAFAnimNode/Source/UAFAnimNode/Public/UAF/AnimNodes/UAFScalePlayRate.h
+++ b/Engine/Plugins/Experimental/UAF/UAFAnimNode/Source/UAFAnimNode/Public/UAF/AnimNodes/UAFScalePlayRate.h
@@ -31,6 +31,10 @@ namespace UE::UAF
 	{
 		UE_API FUAFScalePlayRate(FUAFAnimGraphUpdateContext& Context, const FUAFScalePlayRateData& InData);
 
+		// Sets the play rate
+		// This clears any existing bindings
+		UE_API void SetPlayRate(float PlayRate);
+
 		// FUAFAnimNode impl
 		UE_API virtual void PreUpdate(FUAFAnimGraphUpdateContext& Context) override;
 		UE_API virtual void PostUpdate(FUAFAnimGraphUpdateContext& Context) override;
