Add newly relevant hint to allow nodes to be re-used

#Synced-CL 52678682

diff --git a/Engine/Plugins/Experimental/UAF/UAFAnimNode/Source/UAFAnimNode/Public/UAF/AnimNodeCore/UAFAnimNode.h b/Engine/Plugins/Experimental/UAF/UAFAnimNode/Source/UAFAnimNode/Public/UAF/AnimNodeCore/UAFAnimNode.h
--- a/Engine/Plugins/Experimental/UAF/UAFAnimNode/Source/UAFAnimNode/Public/UAF/AnimNodeCore/UAFAnimNode.h
+++ b/Engine/Plugins/Experimental/UAF/UAFAnimNode/Source/UAFAnimNode/Public/UAF/AnimNodeCore/UAFAnimNode.h
@@ -74,6 +74,11 @@ namespace UE::UAF
 		// If you wish to re-use a node after it has fully blended out, you must first reset it
 		[[nodiscard]] bool IsNewlyRelevant() const;
 
+		// Hints that this node is newly relevant
+		// This should be used when nodes are manually cached and re-used to hint when they become relevant again
+		// Note that if the node is already relevant, it cannot be made irrelevant until it updates
+		void HintNewlyRelevant(bool bIsNewlyRelevant);
+
 		// Returns the AnimOp to perform before our children evaluate theirs
 		[[nodiscard]] FUAFAnimOp* GetPreAnimOp();
 		[[nodiscard]] const FUAFAnimOp* GetPreAnimOp() const;
@@ -343,6 +348,11 @@ namespace UE::UAF
 		return bIsNewlyRelevant;
 	}
 
+	inline void FUAFAnimNode::HintNewlyRelevant(bool bInIsNewlyRelevant)
+	{
+		bIsNewlyRelevant |= bInIsNewlyRelevant;
+	}
+
 	inline FUAFAnimOp* FUAFAnimNode::GetPreAnimOp()
 	{
 		return PreAnimOp;
