Implement blend overwrite and accumulate AnimOps
Implement sanitize AnimOp

#Synced-CL 52678682

diff --git a/Engine/Plugins/Experimental/UAF/UAFAnimNode/Source/UAFAnimNode/Private/UAF/AnimOps/UAFBlendAccumulateAnimOp.cpp b/Engine/Plugins/Experimental/UAF/UAFAnimNode/Source/UAFAnimNode/Private/UAF/AnimOps/UAFBlendAccumulateAnimOp.cpp
new file mode 100644
--- /dev/null
+++ b/Engine/Plugins/Experimental/UAF/UAFAnimNode/Source/UAFAnimNode/Private/UAF/AnimOps/UAFBlendAccumulateAnimOp.cpp
@@ -0,0 +1,50 @@
+// Copyright Epic Games, Inc. All Rights Reserved.
+
+#include "UAF/AnimOps/UAFBlendAccumulateAnimOp.h"
+
+#include "UAF/AnimOpCore/UAFAnimOpValueEvaluator.h"
+#include "UAF/ValueRuntime/Transformers/Accumulate.h"
+
+namespace UE::UAF
+{
+	FUAFBlendAccumulateAnimOp::FUAFBlendAccumulateAnimOp()
+		: FUAFAnimOp(2)
+	{
+		InitializeAs<FUAFBlendAccumulateAnimOp>();
+	}
+
+	void FUAFBlendAccumulateAnimOp::EvaluateValues(FUAFAnimOpValueEvaluator& Evaluator)
+	{
+		// Pop our inputs
+		FPoseValueBundleCoWRef InputRefB = Evaluator.GetEvaluationStack().Pop();
+		FPoseValueBundleCoWRef InputRefA = Evaluator.GetEvaluationStack().Pop();
+
+		const FPoseValueBundle& InputA = InputRefA.Get();
+		const FPoseValueBundle& InputB = InputRefB.Get();
+
+		// Overwrite/accumulate cannot be used with empty inputs because we cannot tell
+		// which space they are supposed to be in until we hit an accumulate but by then
+		// the overwrite weight has been discarded and the remaining weights won't sum to 1.0
+		// To avoid doing a partial blend with an unknown pose, we return an empty result
+		// and let a parent AnimOp handle it
+		if (InputA.IsEmpty())
+		{
+			Evaluator.GetEvaluationStack().Push(MoveTemp(InputRefA));
+		}
+		else if (InputB.IsEmpty())
+		{
+			Evaluator.GetEvaluationStack().Push(MoveTemp(InputRefB));
+		}
+		else
+		{
+			// Named sets must match as it ensures that our inputs have the same sizes/shapes
+			check(InputA.GetNamedSet() == InputB.GetNamedSet());
+
+			FPoseValueBundleCoWRef& OutputRef = FindMutableCoWRef(InputRefB, InputRefA);
+
+			Transformers::FAccumulate::Apply(Evaluator.GetTransformerMap(), InputB, InputA, ScaleFactor, OutputRef.GetMutable());
+
+			Evaluator.GetEvaluationStack().Push(MoveTemp(OutputRef));
+		}
+	}
+}
diff --git a/Engine/Plugins/Experimental/UAF/UAFAnimNode/Source/UAFAnimNode/Private/UAF/AnimOps/UAFBlendOverwriteAnimOp.cpp b/Engine/Plugins/Experimental/UAF/UAFAnimNode/Source/UAFAnimNode/Private/UAF/AnimOps/UAFBlendOverwriteAnimOp.cpp
new file mode 100644
--- /dev/null
+++ b/Engine/Plugins/Experimental/UAF/UAFAnimNode/Source/UAFAnimNode/Private/UAF/AnimOps/UAFBlendOverwriteAnimOp.cpp
@@ -0,0 +1,28 @@
+// Copyright Epic Games, Inc. All Rights Reserved.
+
+#include "UAF/AnimOps/UAFBlendOverwriteAnimOp.h"
+
+#include "UAF/AnimOpCore/UAFAnimOpValueEvaluator.h"
+#include "UAF/ValueRuntime/Transformers/Overwrite.h"
+
+namespace UE::UAF
+{
+	FUAFBlendOverwriteAnimOp::FUAFBlendOverwriteAnimOp()
+		: FUAFAnimOp(1)
+	{
+		InitializeAs<FUAFBlendOverwriteAnimOp>();
+	}
+
+	void FUAFBlendOverwriteAnimOp::EvaluateValues(FUAFAnimOpValueEvaluator& Evaluator)
+	{
+		FPoseValueBundleCoWRef* InputRef = Evaluator.GetEvaluationStack().PeekMutable();
+		FPoseValueBundle& Input = InputRef->GetMutable();
+
+		// We cannot scale an empty input, let the next AnimOp handle it
+		if (!Input.IsEmpty())
+		{
+			// Overwrite in place
+			Transformers::FOverwrite::Apply(Evaluator.GetTransformerMap(), Input, ScaleFactor, Input);
+		}
+	}
+}
diff --git a/Engine/Plugins/Experimental/UAF/UAFAnimNode/Source/UAFAnimNode/Private/UAF/AnimOps/UAFSanitizeAnimOp.cpp b/Engine/Plugins/Experimental/UAF/UAFAnimNode/Source/UAFAnimNode/Private/UAF/AnimOps/UAFSanitizeAnimOp.cpp
new file mode 100644
--- /dev/null
+++ b/Engine/Plugins/Experimental/UAF/UAFAnimNode/Source/UAFAnimNode/Private/UAF/AnimOps/UAFSanitizeAnimOp.cpp
@@ -0,0 +1,25 @@
+// Copyright Epic Games, Inc. All Rights Reserved.
+
+#include "UAF/AnimOps/UAFSanitizeAnimOp.h"
+
+#include "UAF/AnimOpCore/UAFAnimOpValueEvaluator.h"
+#include "UAF/ValueRuntime/Transformers/Sanitize.h"
+
+namespace UE::UAF
+{
+	FUAFSanitizeAnimOp::FUAFSanitizeAnimOp()
+		: FUAFAnimOp(1)
+	{
+		InitializeAs<FUAFSanitizeAnimOp>();
+	}
+
+	void FUAFSanitizeAnimOp::EvaluateValues(FUAFAnimOpValueEvaluator& Evaluator)
+	{
+		FPoseValueBundleCoWRef* InputRef = Evaluator.GetEvaluationStack().PeekMutable();
+
+		if (!InputRef->Get().IsEmpty())
+		{
+			Transformers::FSanitize::Apply(Evaluator.GetTransformerMap(), InputRef->GetMutable());
+		}
+	}
+}
diff --git a/Engine/Plugins/Experimental/UAF/UAFAnimNode/Source/UAFAnimNode/Public/UAF/AnimOps/UAFBlendAccumulateAnimOp.h b/Engine/Plugins/Experimental/UAF/UAFAnimNode/Source/UAFAnimNode/Public/UAF/AnimOps/UAFBlendAccumulateAnimOp.h
new file mode 100644
--- /dev/null
+++ b/Engine/Plugins/Experimental/UAF/UAFAnimNode/Source/UAFAnimNode/Public/UAF/AnimOps/UAFBlendAccumulateAnimOp.h
@@ -0,0 +1,60 @@
+// Copyright Epic Games, Inc. All Rights Reserved.
+
+#pragma once
+
+#include "CoreMinimal.h"
+#include "UAF/AnimOpCore/UAFAnimOp.h"
+
+#include "UAFBlendAccumulateAnimOp.generated.h"
+
+#define UE_API UAFANIMNODE_API
+
+namespace UE::UAF
+{
+	/**
+	 * FUAFBlendAccumulateAnimOp
+	 *
+	 * This pops the top two evaluation stack results (A and B, let B be at the top).
+	 * B is the input that we accumulate on top of; while A is the input we scale.
+	 * The result is pushed back onto the stack.
+	 * Top = B + (A * ScaleFactor)
+	 * Top = Top + (Top-1 * ScaleFactor)
+	 */
+	USTRUCT()
+	struct FUAFBlendAccumulateAnimOp : public FUAFAnimOp
+	{
+		GENERATED_BODY()
+		UAF_DECLARE_ANIMOP(FUAFBlendAccumulateAnimOp)
+
+		UE_API FUAFBlendAccumulateAnimOp();
+
+		// Returns the scale factor to apply to the top-1 result of the evaluation stack
+		float GetScaleFactor() const;
+
+		// Sets the scale factor to apply to the top-1 result of the evaluation stack
+		void SetScaleFactor(float ScaleFactor);
+
+		// FUAFAnimOp impl
+		UE_API virtual void EvaluateValues(FUAFAnimOpValueEvaluator& Evaluator) override;
+
+	private:
+		// The scale factor to apply to the keyframe
+		UPROPERTY(VisibleAnywhere, Category = Properties)
+		float ScaleFactor = 0.0f;
+	};
+
+	//////////////////////////////////////////////////////////////////////////
+	// Inline implementation
+
+	inline float FUAFBlendAccumulateAnimOp::GetScaleFactor() const
+	{
+		return ScaleFactor;
+	}
+
+	inline void FUAFBlendAccumulateAnimOp::SetScaleFactor(float InScaleFactor)
+	{
+		ScaleFactor = InScaleFactor;
+	}
+}
+
+#undef UE_API
diff --git a/Engine/Plugins/Experimental/UAF/UAFAnimNode/Source/UAFAnimNode/Public/UAF/AnimOps/UAFBlendOverwriteAnimOp.h b/Engine/Plugins/Experimental/UAF/UAFAnimNode/Source/UAFAnimNode/Public/UAF/AnimOps/UAFBlendOverwriteAnimOp.h
new file mode 100644
--- /dev/null
+++ b/Engine/Plugins/Experimental/UAF/UAFAnimNode/Source/UAFAnimNode/Public/UAF/AnimOps/UAFBlendOverwriteAnimOp.h
@@ -0,0 +1,56 @@
+// Copyright Epic Games, Inc. All Rights Reserved.
+
+#pragma once
+
+#include "CoreMinimal.h"
+#include "UAF/AnimOpCore/UAFAnimOp.h"
+
+#include "UAFBlendOverwriteAnimOp.generated.h"
+
+#define UE_API UAFANIMNODE_API
+
+namespace UE::UAF
+{
+	/**
+	 * FUAFBlendOverwriteAnimOp
+	 *
+	 * This pops the top result from the evaluation stack, scales it by a weight factor, then pushes back the result.
+	 */
+	USTRUCT()
+	struct FUAFBlendOverwriteAnimOp : public FUAFAnimOp
+	{
+		GENERATED_BODY()
+		UAF_DECLARE_ANIMOP(FUAFBlendOverwriteAnimOp)
+
+		UE_API FUAFBlendOverwriteAnimOp();
+
+		// Returns the scale factor to apply to the top result of the evaluation stack
+		float GetScaleFactor() const;
+
+		// Sets the scale factor to apply to the top result of the evaluation stack
+		void SetScaleFactor(float ScaleFactor);
+
+		// FUAFAnimOp impl
+		UE_API virtual void EvaluateValues(FUAFAnimOpValueEvaluator& Evaluator) override;
+
+	private:
+		// The scale factor to apply to the top result of the evaluations tack
+		UPROPERTY(VisibleAnywhere, Category = Properties)
+		float ScaleFactor = 1.0f;
+	};
+
+	//////////////////////////////////////////////////////////////////////////
+	// Inline implementation
+
+	inline float FUAFBlendOverwriteAnimOp::GetScaleFactor() const
+	{
+		return ScaleFactor;
+	}
+
+	inline void FUAFBlendOverwriteAnimOp::SetScaleFactor(float InScaleFactor)
+	{
+		ScaleFactor = InScaleFactor;
+	}
+}
+
+#undef UE_API
diff --git a/Engine/Plugins/Experimental/UAF/UAFAnimNode/Source/UAFAnimNode/Public/UAF/AnimOps/UAFSanitizeAnimOp.h b/Engine/Plugins/Experimental/UAF/UAFAnimNode/Source/UAFAnimNode/Public/UAF/AnimOps/UAFSanitizeAnimOp.h
new file mode 100644
--- /dev/null
+++ b/Engine/Plugins/Experimental/UAF/UAFAnimNode/Source/UAFAnimNode/Public/UAF/AnimOps/UAFSanitizeAnimOp.h
@@ -0,0 +1,34 @@
+// Copyright Epic Games, Inc. All Rights Reserved.
+
+#pragma once
+
+#include "CoreMinimal.h"
+#include "UAF/AnimOpCore/UAFAnimOp.h"
+
+#include "UAFSanitizeAnimOp.generated.h"
+
+#define UE_API UAFANIMNODE_API
+
+namespace UE::UAF
+{
+	/**
+	 * FUAFSanitizeAnimOp
+	 *
+	 * Sanitizes an input value.
+	 * For pose values, rotations are normalized.
+	 * For float values, we validate that values are not NaN/Inf.
+	 */
+	USTRUCT()
+	struct FUAFSanitizeAnimOp : public FUAFAnimOp
+	{
+		GENERATED_BODY()
+		UAF_DECLARE_ANIMOP(FUAFSanitizeAnimOp)
+
+		UE_API FUAFSanitizeAnimOp();
+
+		// FUAFAnimOp impl
+		UE_API virtual void EvaluateValues(FUAFAnimOpValueEvaluator& Evaluator) override;
+	};
+}
+
+#undef UE_API
